Fixing tooltips & front facing variable names on AudioGameplayVolume & FilterVolumeComponent

diff --git a/Engine/Plugins/AudioGameplayVolume/Source/AudioGameplayVolume/Private/AudioGameplayVolume.cpp b/Engine/Plugins/AudioGameplayVolume/Source/AudioGameplayVolume/Private/AudioGameplayVolume.cpp
--- a/Engine/Plugins/AudioGameplayVolume/Source/AudioGameplayVolume/Private/AudioGameplayVolume.cpp
+++ b/Engine/Plugins/AudioGameplayVolume/Source/AudioGameplayVolume/Private/AudioGameplayVolume.cpp
@@ -28,11 +28,11 @@ AAudioGameplayVolume::AAudioGameplayVolume(const FObjectInitializer& ObjectIniti
 #endif // WITH_EDITOR
 }
 
-void AAudioGameplayVolume::SetEnabled(bool bNewEnabled)
+void AAudioGameplayVolume::SetEnabled(bool bEnable)
 {
-	if (bNewEnabled != bEnabled)
+	if (bEnable != bEnabled)
 	{
-		bEnabled = bNewEnabled;
+		bEnabled = bEnable;
 		if (CanSupportProxy())
 		{
 			AddProxy();
diff --git a/Engine/Plugins/AudioGameplayVolume/Source/AudioGameplayVolume/Public/AudioGameplayVolume.h b/Engine/Plugins/AudioGameplayVolume/Source/AudioGameplayVolume/Public/AudioGameplayVolume.h
--- a/Engine/Plugins/AudioGameplayVolume/Source/AudioGameplayVolume/Public/AudioGameplayVolume.h
+++ b/Engine/Plugins/AudioGameplayVolume/Source/AudioGameplayVolume/Public/AudioGameplayVolume.h
@@ -33,7 +33,7 @@ public:
 	bool GetEnabled() const { return bEnabled; }
 	
 	UFUNCTION(BlueprintCallable, Category = "AudioGameplay")
-	void SetEnabled(bool bNewEnabled);
+	void SetEnabled(bool bEnable);
 
 	//~ Begin UObject Interface
 #if WITH_EDITOR
diff --git a/Engine/Plugins/AudioGameplayVolume/Source/AudioGameplayVolume/Public/FilterVolumeComponent.h b/Engine/Plugins/AudioGameplayVolume/Source/AudioGameplayVolume/Public/FilterVolumeComponent.h
--- a/Engine/Plugins/AudioGameplayVolume/Source/AudioGameplayVolume/Public/FilterVolumeComponent.h
+++ b/Engine/Plugins/AudioGameplayVolume/Source/AudioGameplayVolume/Public/FilterVolumeComponent.h
@@ -57,19 +57,19 @@ private:
 	virtual void FillMutator(TSharedPtr<FProxyVolumeMutator> Mutator) const override;
 	//~ End UAudioGameplayVolumeComponentBase interface
 
-	// The desired volume of sounds outside the volume when the player is inside the volume
+	// The desired LPF frequency cutoff (in hertz) of sounds outside the volume when the player is inside the volume
 	UPROPERTY(EditAnywhere, BlueprintReadWrite, Category = "VolumeFilter", meta = (AllowPrivateAccess = "true"))
 	float ExteriorLPF = MAX_FILTER_FREQUENCY;
 
-	// The time over which to interpolate from the current volume to the desired volume of sounds outside the volume when the player enters the volume
+	// The time over which to interpolate from the current LPF to the desired LPF of sounds outside the volume when the player enters the volume
 	UPROPERTY(EditAnywhere, BlueprintReadWrite, Category = "VolumeFilter", meta = (AllowPrivateAccess = "true"))
 	float ExteriorLPFTime = 0.5f;
 
-	// The desired volume of sounds inside the volume when the player is outside the volume
+	// The desired LPF frequency cutoff (in hertz) of sounds inside the volume when the player is outside the volume
 	UPROPERTY(EditAnywhere, BlueprintReadWrite, Category = "VolumeFilter", meta = (AllowPrivateAccess = "true"))
 	float InteriorLPF = MAX_FILTER_FREQUENCY;
 
-	// The time over which to interpolate from the current volume to the desired volume of sounds inside the volume when the player enters the volume
+	// The time over which to interpolate from the current LPF to the desired LPF of sounds inside the volume when the player enters the volume
 	UPROPERTY(EditAnywhere, BlueprintReadWrite, Category = "VolumeFilter", meta = (AllowPrivateAccess = "true"))
 	float InteriorLPFTime = 0.5f;
 };
\ No newline at end of file
