Extending Blueprintable AudioGameplayVolumes for Listener Enter and Exit events.

diff --git a/Engine/Plugins/AudioGameplayVolume/Source/AudioGameplayVolume/Private/AudioGameplayVolume.cpp b/Engine/Plugins/AudioGameplayVolume/Source/AudioGameplayVolume/Private/AudioGameplayVolume.cpp
--- a/Engine/Plugins/AudioGameplayVolume/Source/AudioGameplayVolume/Private/AudioGameplayVolume.cpp
+++ b/Engine/Plugins/AudioGameplayVolume/Source/AudioGameplayVolume/Private/AudioGameplayVolume.cpp
@@ -143,6 +143,11 @@ void AAudioGameplayVolume::AddProxy() const
 	if (UAudioGameplayVolumeSubsystem* VolumeSubsystem = GetSubsystem())
 	{
 		VolumeSubsystem->AddVolumeComponent(AGVComponent);
+		if (AGVComponent)
+		{
+			AGVComponent->OnProxyEnter.AddUniqueDynamic(this, &AAudioGameplayVolume::OnListenerEnter);
+			AGVComponent->OnProxyExit.AddUniqueDynamic(this, &AAudioGameplayVolume::OnListenerExit);
+		}
 	}
 }
 
@@ -151,6 +156,11 @@ void AAudioGameplayVolume::RemoveProxy() const
 	if (UAudioGameplayVolumeSubsystem* VolumeSubsystem = GetSubsystem())
 	{
 		VolumeSubsystem->RemoveVolumeComponent(AGVComponent);
+		if (AGVComponent)
+		{
+			AGVComponent->OnProxyEnter.RemoveAll(this);
+			AGVComponent->OnProxyExit.RemoveAll(this);
+		}
 	}
 }
 
diff --git a/Engine/Plugins/AudioGameplayVolume/Source/AudioGameplayVolume/Public/AudioGameplayVolume.h b/Engine/Plugins/AudioGameplayVolume/Source/AudioGameplayVolume/Public/AudioGameplayVolume.h
--- a/Engine/Plugins/AudioGameplayVolume/Source/AudioGameplayVolume/Public/AudioGameplayVolume.h
+++ b/Engine/Plugins/AudioGameplayVolume/Source/AudioGameplayVolume/Public/AudioGameplayVolume.h
@@ -35,6 +35,14 @@ public:
 	UFUNCTION(BlueprintCallable, Category = "AudioGameplay")
 	void SetEnabled(bool bEnable);
 
+	/** Blueprint event for listener enter */
+	UFUNCTION(BlueprintImplementableEvent, Category = Events)
+	void OnListenerEnter();
+
+	/** Blueprint event for listener exit */
+	UFUNCTION(BlueprintImplementableEvent, Category = Events)
+	void OnListenerExit();
+
 	//~ Begin UObject Interface
 #if WITH_EDITOR
 	virtual void PostEditChangeProperty(FPropertyChangedEvent& PropertyChangedEvent) override;
diff --git a/Engine/Plugins/AudioGameplayVolume/Source/AudioGameplayVolume/Public/AudioGameplayVolumeComponent.h b/Engine/Plugins/AudioGameplayVolume/Source/AudioGameplayVolume/Public/AudioGameplayVolumeComponent.h
--- a/Engine/Plugins/AudioGameplayVolume/Source/AudioGameplayVolume/Public/AudioGameplayVolumeComponent.h
+++ b/Engine/Plugins/AudioGameplayVolume/Source/AudioGameplayVolume/Public/AudioGameplayVolumeComponent.h
@@ -52,7 +52,7 @@ DECLARE_DYNAMIC_MULTICAST_DELEGATE(FOnAudioGameplayVolumeProxyStateChange);
  *  UAudioGameplayVolumeProxyComponent - Component used to drive interaction with AudioGameplayVolumeSubsystem.
  *   NOTE: Do not inherit from this class, use UAudioGameplayVolumeComponentBase to create extendable functionality
  */
-UCLASS(Config = Game, ClassGroup = ("AudioGameplay"), meta = (BlueprintSpawnableComponent, DisplayName = "Volume Proxy"))
+UCLASS(Config = Game, ClassGroup = ("AudioGameplay"), meta = (BlueprintSpawnableComponent, IsBlueprintBase = false, DisplayName = "Volume Proxy"))
 class AUDIOGAMEPLAYVOLUME_API UAudioGameplayVolumeProxyComponent final : public UAudioGameplayComponent
 {
 	GENERATED_UCLASS_BODY()
@@ -73,20 +73,20 @@ public:
 	/** Called when the proxy is 'exited' - This is when the proxy goes from at least one listeners to zero. */
 	void ExitProxy() const;
 
-protected:
-
-	// A representation of this volume for the audio thread
-	UPROPERTY(Instanced, EditAnywhere, BlueprintReadWrite, Category = "AudioGameplay", Meta = (ShowOnlyInnerProperties, AllowPrivateAccess = "true"))
-	UAudioGameplayVolumeProxy* Proxy = nullptr;
-
 	/** Blueprint event for proxy enter */
-	UPROPERTY(BlueprintAssignable, Category = Events, meta = (AllowPrivateAccess = "true"))
+	UPROPERTY(BlueprintAssignable, Category = Events)
 	FOnAudioGameplayVolumeProxyStateChange OnProxyEnter;
 
 	/** Blueprint event for proxy exit */
-	UPROPERTY(BlueprintAssignable, Category = Events, meta = (AllowPrivateAccess = "true"))
+	UPROPERTY(BlueprintAssignable, Category = Events)
 	FOnAudioGameplayVolumeProxyStateChange OnProxyExit;
 
+protected:
+
+	// A representation of this volume for the audio thread
+	UPROPERTY(Instanced, EditAnywhere, BlueprintReadWrite, Category = "AudioGameplay", Meta = (ShowOnlyInnerProperties, AllowPrivateAccess = "true"))
+	UAudioGameplayVolumeProxy* Proxy = nullptr;
+
 #if WITH_EDITOR
 	virtual void PostEditChangeProperty(FPropertyChangedEvent& PropertyChangedEvent) override;
 #endif // WITH_EDITOR
