Fixing Enable on AudioGameplayVolumes & AutoActivate on AGV components

diff --git a/Engine/Plugins/AudioGameplayVolume/Source/AudioGameplayVolume/Private/AudioGameplayVolume.cpp b/Engine/Plugins/AudioGameplayVolume/Source/AudioGameplayVolume/Private/AudioGameplayVolume.cpp
--- a/Engine/Plugins/AudioGameplayVolume/Source/AudioGameplayVolume/Private/AudioGameplayVolume.cpp
+++ b/Engine/Plugins/AudioGameplayVolume/Source/AudioGameplayVolume/Private/AudioGameplayVolume.cpp
@@ -3,7 +3,6 @@
 #include "AudioGameplayVolume.h"
 #include "AudioGameplayVolumeLogs.h"
 #include "AudioGameplayVolumeProxy.h"
-#include "AudioGameplayVolumeSubsystem.h"
 #include "AudioGameplayVolumeComponent.h"
 #include "AudioDevice.h"
 #include "Components/BrushComponent.h"
@@ -69,6 +68,16 @@ void AAudioGameplayVolume::GetLifetimeReplicatedProps(TArray< FLifetimeProperty
 	DOREPLIFETIME(AAudioGameplayVolume, bEnabled);
 }
 
+void AAudioGameplayVolume::PostInitializeComponents()
+{
+	Super::PostInitializeComponents();
+
+	if (CanSupportProxy())
+	{
+		AddProxy();
+	}
+}
+
 void AAudioGameplayVolume::PostRegisterAllComponents()
 {
 	Super::PostRegisterAllComponents();
@@ -82,11 +91,7 @@ void AAudioGameplayVolume::PostRegisterAllComponents()
 	{
 		UAGVPrimitiveComponentProxy* PrimitiveComponentProxy = NewObject<UAGVPrimitiveComponentProxy>(AGVComponent);
 		AGVComponent->SetProxy(PrimitiveComponentProxy);
-	}
-
-	if (CanSupportProxy())
-	{
-		AddProxy();
+		AGVComponent->bAutoActivate = false;
 	}
 }
 
@@ -140,44 +145,28 @@ void AAudioGameplayVolume::TransformUpdated(USceneComponent* InRootComponent, EU
 
 void AAudioGameplayVolume::AddProxy() const
 {
-	if (UAudioGameplayVolumeSubsystem* VolumeSubsystem = GetSubsystem())
+	if (AGVComponent)
 	{
-		VolumeSubsystem->AddVolumeComponent(AGVComponent);
-		if (AGVComponent)
-		{
-			AGVComponent->OnProxyEnter.AddUniqueDynamic(this, &AAudioGameplayVolume::OnListenerEnter);
-			AGVComponent->OnProxyExit.AddUniqueDynamic(this, &AAudioGameplayVolume::OnListenerExit);
-		}
+		AGVComponent->OnProxyEnter.AddUniqueDynamic(this, &AAudioGameplayVolume::OnListenerEnter);
+		AGVComponent->OnProxyExit.AddUniqueDynamic(this, &AAudioGameplayVolume::OnListenerExit);
+		AGVComponent->Activate();
 	}
 }
 
 void AAudioGameplayVolume::RemoveProxy() const
 {
-	if (UAudioGameplayVolumeSubsystem* VolumeSubsystem = GetSubsystem())
+	if (AGVComponent)
 	{
-		VolumeSubsystem->RemoveVolumeComponent(AGVComponent);
-		if (AGVComponent)
-		{
-			AGVComponent->OnProxyEnter.RemoveAll(this);
-			AGVComponent->OnProxyExit.RemoveAll(this);
-		}
+		AGVComponent->Deactivate();
+		AGVComponent->OnProxyEnter.RemoveAll(this);
+		AGVComponent->OnProxyExit.RemoveAll(this);
 	}
 }
 
 void AAudioGameplayVolume::UpdateProxy() const
 {
-	if (UAudioGameplayVolumeSubsystem* VolumeSubsystem = GetSubsystem())
+	if (AGVComponent)
 	{
-		VolumeSubsystem->UpdateVolumeComponent(AGVComponent);
+		AGVComponent->OnComponentDataChanged();
 	}
 }
-
-UAudioGameplayVolumeSubsystem* AAudioGameplayVolume::GetSubsystem() const
-{
-	if (UWorld* World = GetWorld())
-	{
-		return FAudioDevice::GetSubsystem<UAudioGameplayVolumeSubsystem>(World->GetAudioDevice());
-	}
-
-	return nullptr;
-}
\ No newline at end of file
diff --git a/Engine/Plugins/AudioGameplayVolume/Source/AudioGameplayVolume/Private/AudioGameplayVolumeProxy.cpp b/Engine/Plugins/AudioGameplayVolume/Source/AudioGameplayVolume/Private/AudioGameplayVolumeProxy.cpp
--- a/Engine/Plugins/AudioGameplayVolume/Source/AudioGameplayVolume/Private/AudioGameplayVolumeProxy.cpp
+++ b/Engine/Plugins/AudioGameplayVolume/Source/AudioGameplayVolume/Private/AudioGameplayVolumeProxy.cpp
@@ -35,7 +35,7 @@ void UAudioGameplayVolumeProxy::InitFromComponent(const UAudioGameplayVolumeProx
 	TInlineComponentArray<UAudioGameplayVolumeComponentBase*> Components(Component->GetOwner());
 	for (UAudioGameplayVolumeComponentBase* Comp : Components)
 	{
-		if (!Comp)
+		if (!Comp || !Comp->IsActive())
 		{
 			continue;
 		}
diff --git a/Engine/Plugins/AudioGameplayVolume/Source/AudioGameplayVolume/Public/AudioGameplayVolume.h b/Engine/Plugins/AudioGameplayVolume/Source/AudioGameplayVolume/Public/AudioGameplayVolume.h
--- a/Engine/Plugins/AudioGameplayVolume/Source/AudioGameplayVolume/Public/AudioGameplayVolume.h
+++ b/Engine/Plugins/AudioGameplayVolume/Source/AudioGameplayVolume/Public/AudioGameplayVolume.h
@@ -52,6 +52,7 @@ public:
 	//~ End UObject Interface
 
 	//~ Begin AActor Interface
+	virtual void PostInitializeComponents() override;
 	virtual void PostRegisterAllComponents() override;
 	virtual void PostUnregisterAllComponents() override;
 	//~ End AActor Interface
@@ -71,6 +72,4 @@ protected:
 	void AddProxy() const;
 	void RemoveProxy() const;
 	void UpdateProxy() const;
-
-	UAudioGameplayVolumeSubsystem* GetSubsystem() const;
 };
