Audio Gameplay Volume Fixes
+ Move Volume proxy activation later in lifecycle so other components have a chance to initialize first
+ Add profiling info for Proxy types ::ContainsPosition

diff --git a/Engine/Plugins/AudioGameplayVolume/Source/AudioGameplayVolume/Private/AudioGameplayVolumeProxy.cpp b/Engine/Plugins/AudioGameplayVolume/Source/AudioGameplayVolume/Private/AudioGameplayVolumeProxy.cpp
--- a/Engine/Plugins/AudioGameplayVolume/Source/AudioGameplayVolume/Private/AudioGameplayVolumeProxy.cpp
+++ b/Engine/Plugins/AudioGameplayVolume/Source/AudioGameplayVolume/Private/AudioGameplayVolumeProxy.cpp
@@ -111,6 +111,8 @@ UAGVPrimitiveComponentProxy::UAGVPrimitiveComponentProxy(const FObjectInitialize
 
 bool UAGVPrimitiveComponentProxy::ContainsPosition(const FVector& Position) const
 {
+	SCOPED_NAMED_EVENT(UAGVPrimitiveComponentProxy_ContainsPosition, FColor::Blue);
+
 	FBodyInstance* BodyInstancePointer = nullptr;
 	if (UPrimitiveComponent* PrimitiveComponent = WeakPrimative.Get())
 	{
@@ -151,6 +153,8 @@ UAGVConditionProxy::UAGVConditionProxy(const FObjectInitializer& ObjectInitializ
 
 bool UAGVConditionProxy::ContainsPosition(const FVector& Position) const
 {
+	SCOPED_NAMED_EVENT(UAGVConditionProxy_ContainsPosition, FColor::Blue);
+
 	const UObject* ObjectWithInterface = WeakObject.Get();
 	if (ObjectWithInterface && ObjectWithInterface->Implements<UAudioGameplayCondition>())
 	{
