Audio Gameplay Volume tooltip cleanup

diff --git a/Engine/Plugins/AudioGameplayVolume/Source/AudioGameplayVolume/Public/AudioGameplayVolume.h b/Engine/Plugins/AudioGameplayVolume/Source/AudioGameplayVolume/Public/AudioGameplayVolume.h
--- a/Engine/Plugins/AudioGameplayVolume/Source/AudioGameplayVolume/Public/AudioGameplayVolume.h
+++ b/Engine/Plugins/AudioGameplayVolume/Source/AudioGameplayVolume/Public/AudioGameplayVolume.h
@@ -10,7 +10,11 @@ class UAudioGameplayVolumeProxyComponent;
 class UAudioGameplayVolumeSubsystem;
 
 /**
- * AudioGameplayVolume
+ * AudioGameplayVolume - A spatial volume used to notify audio gameplay systems when the nearest audio listener
+   enters or exits the volume. Additionally, these volumes can influence audio sources depending on the relative
+   position of the listener.
+
+   NOTE: Will only impact audio sources that have "apply ambient volumes" set on their sound class.
  */
 UCLASS(Blueprintable, hidecategories = (Advanced, Attachment, Collision, Volume))
 class AUDIOGAMEPLAYVOLUME_API AAudioGameplayVolume : public AVolume
@@ -32,6 +36,7 @@ public:
 
 	bool GetEnabled() const { return bEnabled; }
 	
+	/** Sets whether the volume is enabled */
 	UFUNCTION(BlueprintCallable, Category = "AudioGameplay")
 	void SetEnabled(bool bEnable);
 
diff --git a/Engine/Source/Runtime/Engine/Classes/Sound/AudioVolume.h b/Engine/Source/Runtime/Engine/Classes/Sound/AudioVolume.h
--- a/Engine/Source/Runtime/Engine/Classes/Sound/AudioVolume.h
+++ b/Engine/Source/Runtime/Engine/Classes/Sound/AudioVolume.h
@@ -26,10 +26,10 @@ class USoundSubmix;
 UENUM(BlueprintType)
 enum class EAudioVolumeLocationState : uint8
 {
-	// A send based on linear interpolation between a distance range and send-level range
+	// Used for when the listener is located inside the Audio Volume
 	InsideTheVolume,
 
-	// A send based on a supplied curve
+	// Used for when the listener is located outside the Audio Volume
 	OutsideTheVolume,
 };
 
@@ -46,7 +46,7 @@ struct ENGINE_API FAudioVolumeSubmixSendSettings
 	UPROPERTY()
 	EAudioVolumeLocationState SourceLocationState_DEPRECATED = EAudioVolumeLocationState::InsideTheVolume;
 
-	// Submix send array for sounds that are outside the audio volume when the listener is inside the volume
+	// Submix send array for sounds that are in the ListenerLocationState at the same time as the listener
 	UPROPERTY(EditAnywhere, BlueprintReadWrite, Category = AudioVolumeSubmixSends)
 	TArray<FSoundSubmixSendInfo> SubmixSends;
 };
