Adding name to filter and attenuation volume proxies

diff --git a/Engine/Plugins/AudioGameplayVolume/Source/AudioGameplayVolume/Private/AttenuationVolumeComponent.cpp b/Engine/Plugins/AudioGameplayVolume/Source/AudioGameplayVolume/Private/AttenuationVolumeComponent.cpp
--- a/Engine/Plugins/AudioGameplayVolume/Source/AudioGameplayVolume/Private/AttenuationVolumeComponent.cpp
+++ b/Engine/Plugins/AudioGameplayVolume/Source/AudioGameplayVolume/Private/AttenuationVolumeComponent.cpp
@@ -4,6 +4,13 @@
 #include "ActiveSound.h"
 #include "AudioGameplayFlags.h"
 
+constexpr TCHAR FProxyMutator_Attenuation::MutatorAttenuationName[];
+
+FProxyMutator_Attenuation::FProxyMutator_Attenuation()
+{
+	MutatorName = MutatorAttenuationName;
+}
+
 void FProxyMutator_Attenuation::Apply(FInteriorSettings& InteriorSettings) const
 {
 	FProxyVolumeMutator::Apply(InteriorSettings);
diff --git a/Engine/Plugins/AudioGameplayVolume/Source/AudioGameplayVolume/Private/FilterVolumeComponent.cpp b/Engine/Plugins/AudioGameplayVolume/Source/AudioGameplayVolume/Private/FilterVolumeComponent.cpp
--- a/Engine/Plugins/AudioGameplayVolume/Source/AudioGameplayVolume/Private/FilterVolumeComponent.cpp
+++ b/Engine/Plugins/AudioGameplayVolume/Source/AudioGameplayVolume/Private/FilterVolumeComponent.cpp
@@ -4,6 +4,13 @@
 #include "ActiveSound.h"
 #include "AudioGameplayFlags.h"
 
+constexpr TCHAR FProxyMutator_Filter::MutatorFilterName[];
+
+FProxyMutator_Filter::FProxyMutator_Filter()
+{
+	MutatorName = MutatorFilterName;
+}
+
 void FProxyMutator_Filter::Apply(FInteriorSettings& InteriorSettings) const
 {
 	FProxyVolumeMutator::Apply(InteriorSettings);
diff --git a/Engine/Plugins/AudioGameplayVolume/Source/AudioGameplayVolume/Public/AttenuationVolumeComponent.h b/Engine/Plugins/AudioGameplayVolume/Source/AudioGameplayVolume/Public/AttenuationVolumeComponent.h
--- a/Engine/Plugins/AudioGameplayVolume/Source/AudioGameplayVolume/Public/AttenuationVolumeComponent.h
+++ b/Engine/Plugins/AudioGameplayVolume/Source/AudioGameplayVolume/Public/AttenuationVolumeComponent.h
@@ -13,7 +13,7 @@ class FProxyMutator_Attenuation : public FProxyVolumeMutator
 {
 public:
 
-	FProxyMutator_Attenuation() = default;
+	FProxyMutator_Attenuation();
 	virtual ~FProxyMutator_Attenuation() = default;
 
 	float ExteriorVolume = 1.0f;
@@ -23,6 +23,10 @@ public:
 
 	virtual void Apply(FInteriorSettings& InteriorSettings) const override;
 	virtual void Apply(FAudioProxyActiveSoundParams& Params) const override;
+
+protected:
+
+	constexpr static const TCHAR MutatorAttenuationName[] = TEXT("Attenuation");
 };
 
 /**
diff --git a/Engine/Plugins/AudioGameplayVolume/Source/AudioGameplayVolume/Public/FilterVolumeComponent.h b/Engine/Plugins/AudioGameplayVolume/Source/AudioGameplayVolume/Public/FilterVolumeComponent.h
--- a/Engine/Plugins/AudioGameplayVolume/Source/AudioGameplayVolume/Public/FilterVolumeComponent.h
+++ b/Engine/Plugins/AudioGameplayVolume/Source/AudioGameplayVolume/Public/FilterVolumeComponent.h
@@ -14,7 +14,7 @@ class FProxyMutator_Filter : public FProxyVolumeMutator
 {
 public:
 
-	FProxyMutator_Filter() = default;
+	FProxyMutator_Filter();
 	virtual ~FProxyMutator_Filter() = default;
 
 	float ExteriorLPF = MAX_FILTER_FREQUENCY;
@@ -24,6 +24,10 @@ public:
 
 	virtual void Apply(FInteriorSettings& InteriorSettings) const override;
 	virtual void Apply(FAudioProxyActiveSoundParams& Params) const override;
+
+protected:
+
+	constexpr static const TCHAR MutatorFilterName[] = TEXT("Filter");
 };
 
 /**
