AudioGameplayVolume - Base Component
+ renamed "FillMutator" to more explicit "CopyAudioDataToMutator" to better reflect functionality
+ restructured creation code path to remove need for Super:: call (because people will easily forget)

diff --git a/Engine/Plugins/AudioGameplayVolume/Source/AudioGameplayVolume/Private/AttenuationVolumeComponent.cpp b/Engine/Plugins/AudioGameplayVolume/Source/AudioGameplayVolume/Private/AttenuationVolumeComponent.cpp
--- a/Engine/Plugins/AudioGameplayVolume/Source/AudioGameplayVolume/Private/AttenuationVolumeComponent.cpp
+++ b/Engine/Plugins/AudioGameplayVolume/Source/AudioGameplayVolume/Private/AttenuationVolumeComponent.cpp
@@ -56,10 +56,8 @@ TSharedPtr<FProxyVolumeMutator> UAttenuationVolumeComponent::FactoryMutator() co
 	return MakeShared<FProxyMutator_Attenuation>();
 }
 
-void UAttenuationVolumeComponent::FillMutator(TSharedPtr<FProxyVolumeMutator> Mutator) const
+void UAttenuationVolumeComponent::CopyAudioDataToMutator(TSharedPtr<FProxyVolumeMutator>& Mutator) const
 {
-	Super::FillMutator(Mutator);
-
 	TSharedPtr<FProxyMutator_Attenuation> AttenuationMutator = StaticCastSharedPtr<FProxyMutator_Attenuation>(Mutator);
 	AttenuationMutator->ExteriorVolume = ExteriorVolume;
 	AttenuationMutator->ExteriorTime = ExteriorTime;
diff --git a/Engine/Plugins/AudioGameplayVolume/Source/AudioGameplayVolume/Private/AudioGameplayVolumeComponent.cpp b/Engine/Plugins/AudioGameplayVolume/Source/AudioGameplayVolume/Private/AudioGameplayVolumeComponent.cpp
--- a/Engine/Plugins/AudioGameplayVolume/Source/AudioGameplayVolume/Private/AudioGameplayVolumeComponent.cpp
+++ b/Engine/Plugins/AudioGameplayVolume/Source/AudioGameplayVolume/Private/AudioGameplayVolumeComponent.cpp
@@ -26,7 +26,7 @@ TSharedPtr<FProxyVolumeMutator> UAudioGameplayVolumeComponentBase::CreateMutator
 	TSharedPtr<FProxyVolumeMutator> ProxyMutator = FactoryMutator();
 	if (ProxyMutator.IsValid())
 	{
-		FillMutator(ProxyMutator);
+		CopyAudioDataToMutatorBase(ProxyMutator);
 	}
 
 	return ProxyMutator;
@@ -37,13 +37,6 @@ TSharedPtr<FProxyVolumeMutator> UAudioGameplayVolumeComponentBase::FactoryMutato
 	return TSharedPtr<FProxyVolumeMutator>();
 }
 
-void UAudioGameplayVolumeComponentBase::FillMutator(TSharedPtr<FProxyVolumeMutator> Mutator) const
-{
-	check(Mutator.IsValid());
-	Mutator->PayloadType = PayloadType;
-	Mutator->Priority = Priority;
-}
-
 void UAudioGameplayVolumeComponentBase::NotifyDataChanged() const
 {
 	if (IsActive())
@@ -59,6 +52,15 @@ void UAudioGameplayVolumeComponentBase::NotifyDataChanged() const
 	}
 }
 
+void UAudioGameplayVolumeComponentBase::CopyAudioDataToMutatorBase(TSharedPtr<FProxyVolumeMutator>& Mutator) const
+{
+	check(Mutator.IsValid());
+	Mutator->PayloadType = PayloadType;
+	Mutator->Priority = Priority;
+
+	CopyAudioDataToMutator(Mutator);
+}
+
 UAudioGameplayVolumeProxyComponent::UAudioGameplayVolumeProxyComponent(const FObjectInitializer& ObjectInitializer)
 	: Super(ObjectInitializer)
 {
diff --git a/Engine/Plugins/AudioGameplayVolume/Source/AudioGameplayVolume/Private/FilterVolumeComponent.cpp b/Engine/Plugins/AudioGameplayVolume/Source/AudioGameplayVolume/Private/FilterVolumeComponent.cpp
--- a/Engine/Plugins/AudioGameplayVolume/Source/AudioGameplayVolume/Private/FilterVolumeComponent.cpp
+++ b/Engine/Plugins/AudioGameplayVolume/Source/AudioGameplayVolume/Private/FilterVolumeComponent.cpp
@@ -56,10 +56,8 @@ TSharedPtr<FProxyVolumeMutator> UFilterVolumeComponent::FactoryMutator() const
 	return MakeShared<FProxyMutator_Filter>();
 }
 
-void UFilterVolumeComponent::FillMutator(TSharedPtr<FProxyVolumeMutator> Mutator) const
+void UFilterVolumeComponent::CopyAudioDataToMutator(TSharedPtr<FProxyVolumeMutator>& Mutator) const
 {
-	Super::FillMutator(Mutator);
-
 	TSharedPtr<FProxyMutator_Filter> FilterMutator = StaticCastSharedPtr<FProxyMutator_Filter>(Mutator);
 	FilterMutator->ExteriorLPF = ExteriorLPF;
 	FilterMutator->ExteriorLPFTime = ExteriorLPFTime;
diff --git a/Engine/Plugins/AudioGameplayVolume/Source/AudioGameplayVolume/Private/SubmixSendVolumeComponent.cpp b/Engine/Plugins/AudioGameplayVolume/Source/AudioGameplayVolume/Private/SubmixSendVolumeComponent.cpp
--- a/Engine/Plugins/AudioGameplayVolume/Source/AudioGameplayVolume/Private/SubmixSendVolumeComponent.cpp
+++ b/Engine/Plugins/AudioGameplayVolume/Source/AudioGameplayVolume/Private/SubmixSendVolumeComponent.cpp
@@ -59,10 +59,8 @@ TSharedPtr<FProxyVolumeMutator> USubmixSendVolumeComponent::FactoryMutator() con
 	return MakeShared<FProxyMutator_SubmixSend>();
 }
 
-void USubmixSendVolumeComponent::FillMutator(TSharedPtr<FProxyVolumeMutator> Mutator) const
+void USubmixSendVolumeComponent::CopyAudioDataToMutator(TSharedPtr<FProxyVolumeMutator>& Mutator) const
 {
-	Super::FillMutator(Mutator);
-
 	TSharedPtr<FProxyMutator_SubmixSend> SubmixMutator = StaticCastSharedPtr<FProxyMutator_SubmixSend>(Mutator);
 	SubmixMutator->SubmixSendSettings = SubmixSendSettings;
 }
diff --git a/Engine/Plugins/AudioGameplayVolume/Source/AudioGameplayVolume/Public/AttenuationVolumeComponent.h b/Engine/Plugins/AudioGameplayVolume/Source/AudioGameplayVolume/Public/AttenuationVolumeComponent.h
--- a/Engine/Plugins/AudioGameplayVolume/Source/AudioGameplayVolume/Public/AttenuationVolumeComponent.h
+++ b/Engine/Plugins/AudioGameplayVolume/Source/AudioGameplayVolume/Public/AttenuationVolumeComponent.h
@@ -57,7 +57,7 @@ private:
 
 	//~ Begin UAudioGameplayVolumeComponentBase interface
 	virtual TSharedPtr<FProxyVolumeMutator> FactoryMutator() const override;
-	virtual void FillMutator(TSharedPtr<FProxyVolumeMutator> Mutator) const override;
+	virtual void CopyAudioDataToMutator(TSharedPtr<FProxyVolumeMutator>& Mutator) const override;
 	//~ End UAudioGameplayVolumeComponentBase interface
 
 	// The desired volume of sounds outside the volume when the player is inside the volume
diff --git a/Engine/Plugins/AudioGameplayVolume/Source/AudioGameplayVolume/Public/AudioGameplayVolumeComponent.h b/Engine/Plugins/AudioGameplayVolume/Source/AudioGameplayVolume/Public/AudioGameplayVolumeComponent.h
--- a/Engine/Plugins/AudioGameplayVolume/Source/AudioGameplayVolume/Public/AudioGameplayVolumeComponent.h
+++ b/Engine/Plugins/AudioGameplayVolume/Source/AudioGameplayVolume/Public/AudioGameplayVolumeComponent.h
@@ -35,8 +35,8 @@ protected:
 	/** Create this component's type of mutator */
 	virtual TSharedPtr<FProxyVolumeMutator> FactoryMutator() const;
 
-	/** Fill the mutator with data from our component */
-	virtual void FillMutator(TSharedPtr<FProxyVolumeMutator> Mutator) const;
+	/** Override in child classes to copy additional data needed to mutators */
+	virtual void CopyAudioDataToMutator(TSharedPtr<FProxyVolumeMutator>& Mutator) const {}
 
 	/** Notify our parent volume our proxy may need updating */
 	void NotifyDataChanged() const;
@@ -44,6 +44,11 @@ protected:
 	// The priority of this component.  In the case of overlapping volumes or multiple affecting components, the highest priority is chosen.
 	UPROPERTY(EditAnywhere, BlueprintReadWrite, Category = "AudioGameplay", Meta = (AllowPrivateAccess = "true"))
 	int32 Priority = 0;
+
+private:
+
+	/** Called for you during mutator creation. See CopyAudioDataToMutator for adding data to derived classes */
+	void CopyAudioDataToMutatorBase(TSharedPtr<FProxyVolumeMutator>& Mutator) const;
 };
 
 DECLARE_DYNAMIC_MULTICAST_DELEGATE(FOnAudioGameplayVolumeProxyStateChange);
diff --git a/Engine/Plugins/AudioGameplayVolume/Source/AudioGameplayVolume/Public/FilterVolumeComponent.h b/Engine/Plugins/AudioGameplayVolume/Source/AudioGameplayVolume/Public/FilterVolumeComponent.h
--- a/Engine/Plugins/AudioGameplayVolume/Source/AudioGameplayVolume/Public/FilterVolumeComponent.h
+++ b/Engine/Plugins/AudioGameplayVolume/Source/AudioGameplayVolume/Public/FilterVolumeComponent.h
@@ -58,7 +58,7 @@ private:
 
 	//~ Begin UAudioGameplayVolumeComponentBase interface
 	virtual TSharedPtr<FProxyVolumeMutator> FactoryMutator() const override;
-	virtual void FillMutator(TSharedPtr<FProxyVolumeMutator> Mutator) const override;
+	virtual void CopyAudioDataToMutator(TSharedPtr<FProxyVolumeMutator>& Mutator) const override;
 	//~ End UAudioGameplayVolumeComponentBase interface
 
 	// The desired LPF frequency cutoff (in hertz) of sounds outside the volume when the player is inside the volume
diff --git a/Engine/Plugins/AudioGameplayVolume/Source/AudioGameplayVolume/Public/SubmixSendVolumeComponent.h b/Engine/Plugins/AudioGameplayVolume/Source/AudioGameplayVolume/Public/SubmixSendVolumeComponent.h
--- a/Engine/Plugins/AudioGameplayVolume/Source/AudioGameplayVolume/Public/SubmixSendVolumeComponent.h
+++ b/Engine/Plugins/AudioGameplayVolume/Source/AudioGameplayVolume/Public/SubmixSendVolumeComponent.h
@@ -47,7 +47,7 @@ private:
 
 	//~ Begin UAudioGameplayVolumeComponentBase interface
 	virtual TSharedPtr<FProxyVolumeMutator> FactoryMutator() const override;
-	virtual void FillMutator(TSharedPtr<FProxyVolumeMutator> Mutator) const override;
+	virtual void CopyAudioDataToMutator(TSharedPtr<FProxyVolumeMutator>& Mutator) const override;
 	//~ End UAudioGameplayVolumeComponentBase interface
 
 	/** Submix send settings to use for this component. Allows audio to dynamically send to submixes based on source and listener locations (relative to parent volume.) */
