AudioGameplayVolume System changes
+ AGVListeners now track active proxy mutators in order to allow a removal step
+ AGVListeners track owning audio device id (instead of world id)

diff --git a/Engine/Plugins/AudioGameplayVolume/Source/AudioGameplayVolume/Private/AudioGameplayVolumeListener.cpp b/Engine/Plugins/AudioGameplayVolume/Source/AudioGameplayVolume/Private/AudioGameplayVolumeListener.cpp
--- a/Engine/Plugins/AudioGameplayVolume/Source/AudioGameplayVolume/Private/AudioGameplayVolumeListener.cpp
+++ b/Engine/Plugins/AudioGameplayVolume/Source/AudioGameplayVolume/Private/AudioGameplayVolumeListener.cpp
@@ -55,18 +55,11 @@ void FInterpolatedInteriorSettings::UpdateInteriorValues()
 	ExteriorLPFInterp = Interpolate(CurrentTime, ExteriorLPFEndTime);
 }
 
-void FAudioGameplayVolumeListener::Update(const FAudioProxyMutatorSearchResult& Result, const FVector& InPosition, uint32 InWorldID)
+void FAudioGameplayVolumeListener::Update(const FAudioProxyMutatorSearchResult& Result, const FVector& InPosition, uint32 InDeviceId)
 {
 	check(IsInAudioThread());
 
-	// If we have a different worldID, this is a new listener; we'll need to exit our previous proxies and enter our current ones
-	if (WorldID != InWorldID)
-	{
-		bNewListener = true;
-	}
-
 	Position = InPosition;
-	WorldID = InWorldID;
 
 	Swap(PreviousProxies, CurrentProxies);
 	CurrentProxies = Result.VolumeSet;
@@ -74,14 +67,31 @@ void FAudioGameplayVolumeListener::Update(const FAudioProxyMutatorSearchResult&
 	TSet<uint32> EnteredProxies = CurrentProxies.Difference(PreviousProxies);
 	TSet<uint32> ExitedProxies = PreviousProxies.Difference(CurrentProxies);
 
-	if (bNewListener || EnteredProxies.Num() || ExitedProxies.Num())
+	// Remove the mutators we were previously affected by and add the new ones if:
+	// We are a new listener to this audio device OR if we have entered or exited any proxies
+	bool bMutatorsOutOfDate = (OwningDeviceId != InDeviceId) || EnteredProxies.Num() || ExitedProxies.Num();
+
+	if (bMutatorsOutOfDate)
 	{
+		OwningDeviceId = InDeviceId;
+
+		// Call remove on all the mutators we were previously affected by 
+		for (const TSharedPtr<FProxyVolumeMutator>& Mutator : ActiveMutators)
+		{
+			Mutator->Remove(*this);
+		}
+
+		// Now cache and apply the currently matching mutators
+		ActiveMutators = Result.MatchingMutators;
+		for (const TSharedPtr<FProxyVolumeMutator>& Mutator : ActiveMutators)
+		{
+			Mutator->Apply(*this);
+		}
+
 		// Reapply interior from mutators
 		InteriorSettings.Apply(Result.InteriorSettings);
 	}
 
 	// Update interpolation
 	InteriorSettings.UpdateInteriorValues();
-
-	bNewListener = false;
 }
diff --git a/Engine/Plugins/AudioGameplayVolume/Source/AudioGameplayVolume/Private/AudioGameplayVolumeSubsystem.cpp b/Engine/Plugins/AudioGameplayVolume/Source/AudioGameplayVolume/Private/AudioGameplayVolumeSubsystem.cpp
--- a/Engine/Plugins/AudioGameplayVolume/Source/AudioGameplayVolume/Private/AudioGameplayVolumeSubsystem.cpp
+++ b/Engine/Plugins/AudioGameplayVolume/Source/AudioGameplayVolume/Private/AudioGameplayVolumeSubsystem.cpp
@@ -465,7 +465,7 @@ void UAudioGameplayVolumeSubsystem::UpdateFromListeners()
 
 			// Reassign the set of all proxy volumes (regardless of payload type)
 			Swap(TempVolumeSet, Result.VolumeSet);
-			AGVListeners[i].Update(Result, ProxySearch.Location, ProxySearch.WorldID);
+			AGVListeners[i].Update(Result, ProxySearch.Location, DeviceHandle.GetDeviceID());
 
 			if (FAudioGameplayVolumeProxyInfo* ProxyInfo = WorldProxyLists.Find(ProxySearch.WorldID))
 			{
diff --git a/Engine/Plugins/AudioGameplayVolume/Source/AudioGameplayVolume/Public/AudioGameplayVolumeListener.h b/Engine/Plugins/AudioGameplayVolume/Source/AudioGameplayVolume/Public/AudioGameplayVolumeListener.h
--- a/Engine/Plugins/AudioGameplayVolume/Source/AudioGameplayVolume/Public/AudioGameplayVolumeListener.h
+++ b/Engine/Plugins/AudioGameplayVolume/Source/AudioGameplayVolume/Public/AudioGameplayVolumeListener.h
@@ -5,6 +5,7 @@
 #include "Sound/AudioVolume.h"
 
 // Forward Declarations 
+class FProxyVolumeMutator;
 struct FAudioProxyMutatorSearchResult;
 
 /**
@@ -71,10 +72,12 @@ public:
 
 	FAudioGameplayVolumeListener() = default;
 
-	void Update(const FAudioProxyMutatorSearchResult& Result, const FVector& InPosition, uint32 InWorldID);
+	void Update(const FAudioProxyMutatorSearchResult& Result, const FVector& InPosition, uint32 InDeviceId);
 
 	const TSet<uint32>& GetCurrentProxies() const { return CurrentProxies; }
 
+	uint32 GetOwningDeviceId() const { return OwningDeviceId; }
+
 	bool GetAffectedByLegacySystem() const { return bAffectedByLegacySystem; }
 	void SetAffectedByLegacySystem(bool bIsAffected) { bAffectedByLegacySystem = bIsAffected; }
 
@@ -84,12 +87,14 @@ protected:
 
 	FInterpolatedInteriorSettings InteriorSettings;
 
+	uint32 OwningDeviceId = INDEX_NONE;
 	FVector Position = FVector::ZeroVector;
-	uint32 WorldID = INDEX_NONE;
-	bool bNewListener = true;
+
 	bool bAffectedByLegacySystem = false;
 
 	TSet<uint32> CurrentProxies;
 	TSet<uint32> PreviousProxies;
+
+	TArray<TSharedPtr<FProxyVolumeMutator>> ActiveMutators;
 };
 
diff --git a/Engine/Plugins/AudioGameplayVolume/Source/AudioGameplayVolume/Public/AudioGameplayVolumeProxyMutator.h b/Engine/Plugins/AudioGameplayVolume/Source/AudioGameplayVolume/Public/AudioGameplayVolumeProxyMutator.h
--- a/Engine/Plugins/AudioGameplayVolume/Source/AudioGameplayVolume/Public/AudioGameplayVolumeProxyMutator.h
+++ b/Engine/Plugins/AudioGameplayVolume/Source/AudioGameplayVolume/Public/AudioGameplayVolumeProxyMutator.h
@@ -108,6 +108,10 @@ public:
 
 	virtual void Apply(FInteriorSettings& InteriorSettings) const;
 	virtual void Apply(FAudioProxyActiveSoundParams& Params) const {}
+	virtual void Apply(FAudioGameplayVolumeListener& Listener) const {}
+
+	virtual void Remove(FAudioProxyActiveSoundParams& Params) const {}
+	virtual void Remove(FAudioGameplayVolumeListener& Listener) const {}
 
 	bool HasPayloadType(PayloadFlags InType) const;
 
