Reverb Volume Component for AudioGameplayVolume system.

New component for use with the AGV system.  Allows users to specify reverb effects that are applied (and removed) on entering (or exiting) the associated volume.

diff --git a/Engine/Plugins/AudioGameplayVolume/Source/AudioGameplayVolume/Private/ReverbVolumeComponent.cpp b/Engine/Plugins/AudioGameplayVolume/Source/AudioGameplayVolume/Private/ReverbVolumeComponent.cpp
new file mode 100644
--- /dev/null
+++ b/Engine/Plugins/AudioGameplayVolume/Source/AudioGameplayVolume/Private/ReverbVolumeComponent.cpp
@@ -0,0 +1,81 @@
+// Copyright Epic Games, Inc. All Rights Reserved.
+
+#include "ReverbVolumeComponent.h"
+#include "AudioDevice.h"
+#include "AudioGameplayFlags.h"
+#include "AudioGameplayVolumeListener.h"
+
+constexpr TCHAR FProxyMutator_Reverb::MutatorReverbName[];
+
+FProxyMutator_Reverb::FProxyMutator_Reverb()
+{
+	MutatorName = MutatorReverbName;
+}
+
+void FProxyMutator_Reverb::Apply(FAudioGameplayVolumeListener& Listener) const
+{
+	check(IsInAudioThread());
+
+	FAudioDeviceHandle AudioDeviceHandle = FAudioDeviceManager::Get()->GetAudioDevice(Listener.GetOwningDeviceId());
+	if (AudioDeviceHandle.IsValid())
+	{
+		DECLARE_CYCLE_STAT(TEXT("FAudioThreadTask.AGVActivateReverb"), STAT_AGVActivateReverb, STATGROUP_AudioThreadCommands);
+
+		float PriorityAsFloat = static_cast<float>(Priority);
+		FReverbSettings Settings = ReverbSettings;
+		FAudioThread::RunCommandOnGameThread([AudioDeviceHandle, Settings, PriorityAsFloat]() mutable
+		{
+			if (AudioDeviceHandle.IsValid())
+			{
+				static FName NAME_AGVReverb(TEXT("AGVReverb"));
+				AudioDeviceHandle->ActivateReverbEffect(Settings.ReverbEffect, NAME_AGVReverb, PriorityAsFloat, Settings.Volume, Settings.FadeTime);
+			}
+		}, GET_STATID(STAT_AGVActivateReverb));
+	}
+}
+
+void FProxyMutator_Reverb::Remove(FAudioGameplayVolumeListener& Listener) const
+{
+	check(IsInAudioThread());
+
+	FAudioDeviceHandle AudioDeviceHandle = FAudioDeviceManager::Get()->GetAudioDevice(Listener.GetOwningDeviceId());
+	if (AudioDeviceHandle.IsValid())
+	{
+		DECLARE_CYCLE_STAT(TEXT("FAudioThreadTask.AGVDeativateReverb"), STAT_AGVDeactivateReverb, STATGROUP_AudioThreadCommands);
+
+		FAudioThread::RunCommandOnGameThread([AudioDeviceHandle]() mutable
+		{
+			if (AudioDeviceHandle.IsValid())
+			{
+				static FName NAME_AGVReverb(TEXT("AGVReverb"));
+				AudioDeviceHandle->DeactivateReverbEffect(NAME_AGVReverb);
+			}
+		}, GET_STATID(STAT_AGVDeactivateReverb));
+	}
+}
+
+UReverbVolumeComponent::UReverbVolumeComponent(const FObjectInitializer& ObjectInitializer)
+	: Super(ObjectInitializer)
+{
+	PayloadType = AudioGameplay::EComponentPayload::AGCP_Listener;
+	bAutoActivate = true;
+}
+
+void UReverbVolumeComponent::SetReverbSettings(const FReverbSettings& NewReverbSettings)
+{
+	ReverbSettings = NewReverbSettings;
+
+	// Let the parent volume know we've changed
+	NotifyDataChanged();
+}
+
+TSharedPtr<FProxyVolumeMutator> UReverbVolumeComponent::FactoryMutator() const
+{
+	return MakeShared<FProxyMutator_Reverb>();
+}
+
+void UReverbVolumeComponent::CopyAudioDataToMutator(TSharedPtr<FProxyVolumeMutator>& Mutator) const
+{
+	TSharedPtr<FProxyMutator_Reverb> ReverbMutator = StaticCastSharedPtr<FProxyMutator_Reverb>(Mutator);
+	ReverbMutator->ReverbSettings = ReverbSettings;
+}
diff --git a/Engine/Plugins/AudioGameplayVolume/Source/AudioGameplayVolume/Public/ReverbVolumeComponent.h b/Engine/Plugins/AudioGameplayVolume/Source/AudioGameplayVolume/Public/ReverbVolumeComponent.h
new file mode 100644
--- /dev/null
+++ b/Engine/Plugins/AudioGameplayVolume/Source/AudioGameplayVolume/Public/ReverbVolumeComponent.h
@@ -0,0 +1,57 @@
+// Copyright Epic Games, Inc. All Rights Reserved.
+
+#pragma once
+
+#include "AudioGameplayVolumeComponent.h"
+#include "AudioGameplayVolumeProxyMutator.h"
+#include "Sound/ReverbSettings.h"
+#include "ReverbVolumeComponent.generated.h"
+
+/**
+ *  FProxyMutator_Reverb - An audio thread representation of the reverb volume component
+ */
+class FProxyMutator_Reverb : public FProxyVolumeMutator
+{
+public:
+
+	FProxyMutator_Reverb();
+	virtual ~FProxyMutator_Reverb() = default;
+
+	FReverbSettings ReverbSettings;
+
+	virtual void Apply(FAudioGameplayVolumeListener& Listener) const override;
+	virtual void Remove(FAudioGameplayVolumeListener& Listener) const override;
+
+protected:
+
+	constexpr static const TCHAR MutatorReverbName[] = TEXT("Reverb");
+};
+
+/**
+ *  UReverbVolumeComponent - Audio Gameplay Volume component for reverb settings
+ */
+UCLASS(Blueprintable, Config = Game, ClassGroup = ("AudioGameplayVolume"), meta = (BlueprintSpawnableComponent, DisplayName = "Reverb"))
+class AUDIOGAMEPLAYVOLUME_API UReverbVolumeComponent : public UAudioGameplayVolumeComponentBase
+{
+	GENERATED_UCLASS_BODY()
+
+public:
+
+	virtual ~UReverbVolumeComponent() = default;
+
+	UFUNCTION(BlueprintCallable, Category = "AudioGameplay")
+	void SetReverbSettings(const FReverbSettings& NewReverbSettings);
+
+	const FReverbSettings& GetReverbSettings() const { return ReverbSettings; }
+
+private:
+
+	//~ Begin UAudioGameplayVolumeComponentBase interface
+	virtual TSharedPtr<FProxyVolumeMutator> FactoryMutator() const override;
+	virtual void CopyAudioDataToMutator(TSharedPtr<FProxyVolumeMutator>& Mutator) const override;
+	//~ End UAudioGameplayVolumeComponentBase interface
+
+	/** Reverb settings to use with this component */
+	UPROPERTY(EditAnywhere, BlueprintReadOnly, Category = Reverb, meta = (AllowPrivateAccess = "true"))
+	FReverbSettings ReverbSettings;
+};
\ No newline at end of file
