Submix Override Volume Component for AudioGameplayVolume system.

New component for use with the AGV system.  Allows users to specify submix effect chain overrides that are applied (and removed) on entering (or exiting) the associated volume.

diff --git a/Engine/Plugins/AudioGameplayVolume/Source/AudioGameplayVolume/Private/SubmixOverrideVolumeComponent.cpp b/Engine/Plugins/AudioGameplayVolume/Source/AudioGameplayVolume/Private/SubmixOverrideVolumeComponent.cpp
new file mode 100644
--- /dev/null
+++ b/Engine/Plugins/AudioGameplayVolume/Source/AudioGameplayVolume/Private/SubmixOverrideVolumeComponent.cpp
@@ -0,0 +1,91 @@
+// Copyright Epic Games, Inc. All Rights Reserved.
+
+#include "SubmixOverrideVolumeComponent.h"
+#include "AudioDevice.h"
+#include "AudioGameplayFlags.h"
+#include "AudioGameplayVolumeListener.h"
+
+constexpr TCHAR FProxyMutator_SubmixOverride::MutatorSubmixOverrideName[];
+
+FProxyMutator_SubmixOverride::FProxyMutator_SubmixOverride()
+{
+	MutatorName = MutatorSubmixOverrideName;
+}
+
+void FProxyMutator_SubmixOverride::Apply(FAudioGameplayVolumeListener& Listener) const
+{
+	check(IsInAudioThread());
+
+	FAudioDeviceHandle AudioDeviceHandle = FAudioDeviceManager::Get()->GetAudioDevice(Listener.GetOwningDeviceId());
+	if (AudioDeviceHandle.IsValid() && AudioDeviceHandle->IsAudioMixerEnabled())
+	{
+		FSoundEffectSubmixInitData InitData;
+		InitData.DeviceID = AudioDeviceHandle.GetDeviceID();
+		InitData.SampleRate = AudioDeviceHandle->GetSampleRate();
+		TArray<FSoundEffectSubmixPtr> SubmixEffectPresetChainOverride;
+
+		for (const FAudioVolumeSubmixOverrideSettings& OverrideSettings : SubmixOverrideSettings)
+		{
+			if (OverrideSettings.Submix && OverrideSettings.SubmixEffectChain.Num() > 0)
+			{
+				// Build the instances of the new submix preset chain override
+				for (USoundEffectSubmixPreset* SubmixEffectPreset : OverrideSettings.SubmixEffectChain)
+				{
+					if (SubmixEffectPreset)
+					{
+						InitData.PresetSettings = nullptr;
+						InitData.ParentPresetUniqueId = SubmixEffectPreset->GetUniqueID();
+
+						TSoundEffectSubmixPtr SoundEffectSubmix = USoundEffectPreset::CreateInstance<FSoundEffectSubmixInitData, FSoundEffectSubmix>(InitData, *SubmixEffectPreset);
+						SoundEffectSubmix->SetEnabled(true);
+						SubmixEffectPresetChainOverride.Add(SoundEffectSubmix);
+					}
+				}
+
+				AudioDeviceHandle->SetSubmixEffectChainOverride(OverrideSettings.Submix, SubmixEffectPresetChainOverride, OverrideSettings.CrossfadeTime);
+				SubmixEffectPresetChainOverride.Reset();
+			}
+		}
+	}
+}
+
+void FProxyMutator_SubmixOverride::Remove(FAudioGameplayVolumeListener& Listener) const
+{
+	check(IsInAudioThread());
+
+	FAudioDeviceHandle AudioDeviceHandle = FAudioDeviceManager::Get()->GetAudioDevice(Listener.GetOwningDeviceId());
+	if (AudioDeviceHandle.IsValid() && AudioDeviceHandle->IsAudioMixerEnabled())
+	{
+		// Clear out any previous submix effect chain overrides
+		for (const FAudioVolumeSubmixOverrideSettings& OverrideSettings : SubmixOverrideSettings)
+		{
+			AudioDeviceHandle->ClearSubmixEffectChainOverride(OverrideSettings.Submix, OverrideSettings.CrossfadeTime);
+		}
+	}
+}
+
+USubmixOverrideVolumeComponent::USubmixOverrideVolumeComponent(const FObjectInitializer& ObjectInitializer)
+	: Super(ObjectInitializer)
+{
+	PayloadType = AudioGameplay::EComponentPayload::AGCP_Listener;
+	bAutoActivate = true;
+}
+
+void USubmixOverrideVolumeComponent::SetSubmixOverrideSettings(const TArray<FAudioVolumeSubmixOverrideSettings>& NewSubmixOverrideSettings)
+{
+	SubmixOverrideSettings = NewSubmixOverrideSettings;
+
+	// Let the parent volume know we've changed
+	NotifyDataChanged();
+}
+
+TSharedPtr<FProxyVolumeMutator> USubmixOverrideVolumeComponent::FactoryMutator() const
+{
+	return MakeShared<FProxyMutator_SubmixOverride>();
+}
+
+void USubmixOverrideVolumeComponent::CopyAudioDataToMutator(TSharedPtr<FProxyVolumeMutator>& Mutator) const
+{
+	TSharedPtr<FProxyMutator_SubmixOverride> SubmixMutator = StaticCastSharedPtr<FProxyMutator_SubmixOverride>(Mutator);
+	SubmixMutator->SubmixOverrideSettings = SubmixOverrideSettings;
+}
diff --git a/Engine/Plugins/AudioGameplayVolume/Source/AudioGameplayVolume/Public/SubmixOverrideVolumeComponent.h b/Engine/Plugins/AudioGameplayVolume/Source/AudioGameplayVolume/Public/SubmixOverrideVolumeComponent.h
new file mode 100644
--- /dev/null
+++ b/Engine/Plugins/AudioGameplayVolume/Source/AudioGameplayVolume/Public/SubmixOverrideVolumeComponent.h
@@ -0,0 +1,57 @@
+// Copyright Epic Games, Inc. All Rights Reserved.
+
+#pragma once
+
+#include "AudioGameplayVolumeComponent.h"
+#include "AudioGameplayVolumeProxyMutator.h"
+#include "Sound/AudioVolume.h"
+#include "SubmixOverrideVolumeComponent.generated.h"
+
+/**
+ *  FProxyMutator_SubmixOverride - An audio thread representation of Submix Overrides
+ */
+class FProxyMutator_SubmixOverride : public FProxyVolumeMutator
+{
+public:
+
+	FProxyMutator_SubmixOverride();
+	virtual ~FProxyMutator_SubmixOverride() = default;
+
+	TArray<FAudioVolumeSubmixOverrideSettings> SubmixOverrideSettings;
+
+	virtual void Apply(FAudioGameplayVolumeListener& Listener) const override;
+	virtual void Remove(FAudioGameplayVolumeListener& Listener) const override;
+
+protected:
+
+	constexpr static const TCHAR MutatorSubmixOverrideName[] = TEXT("SubmixOverride");
+};
+
+/**
+ *  USubmixOverrideVolumeComponent - Audio Gameplay Volume component for submix effect chain overrides
+ */
+UCLASS(Blueprintable, Config = Game, ClassGroup = ("AudioGameplayVolume"), meta = (BlueprintSpawnableComponent, DisplayName = "Submix Override"))
+class AUDIOGAMEPLAYVOLUME_API USubmixOverrideVolumeComponent : public UAudioGameplayVolumeComponentBase
+{
+	GENERATED_UCLASS_BODY()
+
+public:
+
+	virtual ~USubmixOverrideVolumeComponent() = default;
+
+	UFUNCTION(BlueprintCallable, Category = "AudioGameplay")
+	void SetSubmixOverrideSettings(const TArray<FAudioVolumeSubmixOverrideSettings>& NewSubmixOverrideSettings);
+
+	const TArray<FAudioVolumeSubmixOverrideSettings>& GetSubmixOverrideSettings() const { return SubmixOverrideSettings; }
+
+private:
+
+	//~ Begin UAudioGameplayVolumeComponentBase interface
+	virtual TSharedPtr<FProxyVolumeMutator> FactoryMutator() const override;
+	virtual void CopyAudioDataToMutator(TSharedPtr<FProxyVolumeMutator>& Mutator) const override;
+	//~ End UAudioGameplayVolumeComponentBase interface
+
+	/** Submix effect chain override settings. Will override the effect chains on the given submixes */
+	UPROPERTY(EditAnywhere, BlueprintReadOnly, Category = Submixes, meta = (AllowPrivateAccess = "true"))
+	TArray<FAudioVolumeSubmixOverrideSettings> SubmixOverrideSettings;
+};
\ No newline at end of file
