Add Listener Enter/Exit bp assignable delegates to AudioGameplayVolume for scripting usage

diff --git a/Engine/Plugins/AudioGameplayVolume/Source/AudioGameplayVolume/Private/AudioGameplayVolume.cpp b/Engine/Plugins/AudioGameplayVolume/Source/AudioGameplayVolume/Private/AudioGameplayVolume.cpp
--- a/Engine/Plugins/AudioGameplayVolume/Source/AudioGameplayVolume/Private/AudioGameplayVolume.cpp
+++ b/Engine/Plugins/AudioGameplayVolume/Source/AudioGameplayVolume/Private/AudioGameplayVolume.cpp
@@ -43,6 +43,16 @@ void AAudioGameplayVolume::SetEnabled(bool bEnable)
 	}
 }
 
+void AAudioGameplayVolume::OnListenerEnter_Implementation()
+{
+	OnListenerEnterEvent.Broadcast();
+}
+
+void AAudioGameplayVolume::OnListenerExit_Implementation()
+{
+	OnListenerExitEvent.Broadcast();
+}
+
 #if WITH_EDITOR
 void AAudioGameplayVolume::PostEditChangeProperty(FPropertyChangedEvent& PropertyChangedEvent)
 {
diff --git a/Engine/Plugins/AudioGameplayVolume/Source/AudioGameplayVolume/Public/AudioGameplayVolume.h b/Engine/Plugins/AudioGameplayVolume/Source/AudioGameplayVolume/Public/AudioGameplayVolume.h
--- a/Engine/Plugins/AudioGameplayVolume/Source/AudioGameplayVolume/Public/AudioGameplayVolume.h
+++ b/Engine/Plugins/AudioGameplayVolume/Source/AudioGameplayVolume/Public/AudioGameplayVolume.h
@@ -5,6 +5,9 @@
 #include "GameFramework/Volume.h"
 #include "AudioGameplayVolume.generated.h"
 
+
+DECLARE_DYNAMIC_MULTICAST_DELEGATE(FAudioGameplayVolumeEvent);
+
 // Forward Declarations 
 class UAudioGameplayVolumeProxyComponent;
 class UAudioGameplayVolumeSubsystem;
@@ -41,13 +44,19 @@ public:
 	void SetEnabled(bool bEnable);
 
 	/** Blueprint event for listener enter */
-	UFUNCTION(BlueprintImplementableEvent, Category = Events)
+	UFUNCTION(BlueprintNativeEvent, Category = Events)
 	void OnListenerEnter();
 
 	/** Blueprint event for listener exit */
-	UFUNCTION(BlueprintImplementableEvent, Category = Events)
+	UFUNCTION(BlueprintNativeEvent, Category = Events)
 	void OnListenerExit();
 
+	UPROPERTY(BlueprintAssignable, Category = Events)
+	FAudioGameplayVolumeEvent OnListenerEnterEvent;
+
+	UPROPERTY(BlueprintAssignable, Category = Events)
+	FAudioGameplayVolumeEvent OnListenerExitEvent;
+
 	//~ Begin UObject Interface
 #if WITH_EDITOR
 	virtual void PostEditChangeProperty(FPropertyChangedEvent& PropertyChangedEvent) override;
