Reorganization pass on AudioGameplayVolumes (AGVs) before 5.1
+ Renamed UAudioGameplayVolumeComponentBase to UAudioGameplayVolumeMutator, as these components are for more advanced effects with AGVs.  Simple component interaction is done with an ActorComponent that implements an interface, not inheritance.
+ Moved UAudioGameplayVolumeMutator and associated implementation into AudioGameplayVolumeMutator header and impl.
+ Moved Proxy search functionality into the subsystem files.
+ Cleaned up unneeded includes, removed (now) empty files.

#PF: 62b798e4d17a95b3eb902ab0

diff --git a/Engine/Plugins/AudioGameplayVolume/Source/AudioGameplayVolume/Private/AudioGameplayVolumeComponent.cpp b/Engine/Plugins/AudioGameplayVolume/Source/AudioGameplayVolume/Private/AudioGameplayVolumeComponent.cpp
--- a/Engine/Plugins/AudioGameplayVolume/Source/AudioGameplayVolume/Private/AudioGameplayVolumeComponent.cpp
+++ b/Engine/Plugins/AudioGameplayVolume/Source/AudioGameplayVolume/Private/AudioGameplayVolumeComponent.cpp
@@ -3,64 +3,9 @@
 #include "AudioGameplayVolumeComponent.h"
 #include "AudioGameplayVolumeSubsystem.h"
 #include "AudioGameplayVolumeProxy.h"
-#include "AudioGameplayVolumeProxyMutator.h"
 #include "Interfaces/IAudioGameplayVolumeInteraction.h"
 #include "AudioDevice.h"
 
-UAudioGameplayVolumeComponentBase::UAudioGameplayVolumeComponentBase(const FObjectInitializer& ObjectInitializer)
-	: Super(ObjectInitializer)
-{
-}
-
-void UAudioGameplayVolumeComponentBase::SetPriority(int32 InPriority)
-{
-	if (Priority != InPriority)
-	{
-		Priority = InPriority;
-		NotifyDataChanged();
-	}
-}
-
-TSharedPtr<FProxyVolumeMutator> UAudioGameplayVolumeComponentBase::CreateMutator() const
-{
-	TSharedPtr<FProxyVolumeMutator> ProxyMutator = FactoryMutator();
-	if (ProxyMutator.IsValid())
-	{
-		CopyAudioDataToMutatorBase(ProxyMutator);
-	}
-
-	return ProxyMutator;
-}
-
-TSharedPtr<FProxyVolumeMutator> UAudioGameplayVolumeComponentBase::FactoryMutator() const
-{
-	return TSharedPtr<FProxyVolumeMutator>();
-}
-
-void UAudioGameplayVolumeComponentBase::NotifyDataChanged() const
-{
-	if (IsActive())
-	{
-		TInlineComponentArray<UAudioGameplayVolumeProxyComponent*> VolumeComponents(GetOwner());
-		if (ensureMsgf(VolumeComponents.Num() == 1, TEXT("Expecting exactly one AudioGameplayVolumeProxyComponent on an actor!")))
-		{
-			if (VolumeComponents[0])
-			{
-				VolumeComponents[0]->OnComponentDataChanged();
-			}
-		}
-	}
-}
-
-void UAudioGameplayVolumeComponentBase::CopyAudioDataToMutatorBase(TSharedPtr<FProxyVolumeMutator>& Mutator) const
-{
-	check(Mutator.IsValid());
-	Mutator->PayloadType = PayloadType;
-	Mutator->Priority = Priority;
-
-	CopyAudioDataToMutator(Mutator);
-}
-
 UAudioGameplayVolumeProxyComponent::UAudioGameplayVolumeProxyComponent(const FObjectInitializer& ObjectInitializer)
 	: Super(ObjectInitializer)
 {
diff --git a/Engine/Plugins/AudioGameplayVolume/Source/AudioGameplayVolume/Private/AudioGameplayVolumeListener.cpp b/Engine/Plugins/AudioGameplayVolume/Source/AudioGameplayVolume/Private/AudioGameplayVolumeListener.cpp
--- a/Engine/Plugins/AudioGameplayVolume/Source/AudioGameplayVolume/Private/AudioGameplayVolumeListener.cpp
+++ b/Engine/Plugins/AudioGameplayVolume/Source/AudioGameplayVolume/Private/AudioGameplayVolumeListener.cpp
@@ -1,7 +1,8 @@
 // Copyright Epic Games, Inc. All Rights Reserved.
 
 #include "AudioGameplayVolumeListener.h"
-#include "AudioGameplayVolumeProxyMutator.h"
+#include "AudioGameplayVolumeMutator.h"
+#include "AudioGameplayVolumeSubsystem.h"
 #include "Misc/App.h"
 
 void FInterpolatedInteriorSettings::Apply(const FInteriorSettings& NewSettings)
diff --git a/Engine/Plugins/AudioGameplayVolume/Source/AudioGameplayVolume/Private/AudioGameplayVolumeMutator.cpp b/Engine/Plugins/AudioGameplayVolume/Source/AudioGameplayVolume/Private/AudioGameplayVolumeMutator.cpp
new file mode 100644
--- /dev/null
+++ b/Engine/Plugins/AudioGameplayVolume/Source/AudioGameplayVolume/Private/AudioGameplayVolumeMutator.cpp
@@ -0,0 +1,182 @@
+// Copyright Epic Games, Inc. All Rights Reserved.
+
+#include "AudioGameplayVolumeMutator.h"
+#include "AudioGameplayVolumeListener.h"
+#include "AudioGameplayVolumeSubsystem.h"
+
+FAudioProxyActiveSoundParams::FAudioProxyActiveSoundParams(const FAudioGameplayActiveSoundInfo& SoundInfo, const FAudioGameplayVolumeListener& InListener)
+	: SourceInteriorVolume(SoundInfo.SourceInteriorVolume)
+	, SourceInteriorLPF(SoundInfo.SourceInteriorLPF)
+	, bUsingWorldSettings(SoundInfo.InteriorSettings.IsWorldSettings())
+	, Listener(InListener)
+	, Sound(SoundInfo)
+{
+}
+
+void FAudioProxyActiveSoundParams::UpdateInteriorValues()
+{
+	const FInterpolatedInteriorSettings& ListenerInterior = Listener.GetInteriorSettings();
+
+	// When a listener or a sound changes interior settings, we smoothly transition from the old settings
+	// to the new settings.  We track both the listener and the sound's interpolation progress.
+	// We use the interpolation progress from the object that has _most recently_ changed interior settings.
+	// This allows moving sounds to smoothly interpolate across interior boundaries after the listener has completed 
+	// transitioning to it's current interior settings (otherwise we'd abruptly slam the moving sound's interior settings to match the listener's).
+	float ExtVolumeInterp = ListenerInterior.GetExteriorVolumeInterp();
+	float IntVolumeInterp = ListenerInterior.GetInteriorVolumeInterp();
+	float ExtLowPassInterp = ListenerInterior.GetExteriorLPFInterp();
+	float IntLowPassInterp = ListenerInterior.GetInteriorLPFInterp();
+
+	if (Sound.InteriorSettings.GetInteriorStartTime() > ListenerInterior.GetInteriorStartTime())
+	{
+		ExtVolumeInterp = Sound.InteriorSettings.GetExteriorVolumeInterp();
+		IntVolumeInterp = Sound.InteriorSettings.GetInteriorVolumeInterp();
+		ExtLowPassInterp = Sound.InteriorSettings.GetExteriorLPFInterp();
+		IntLowPassInterp = Sound.InteriorSettings.GetInteriorLPFInterp();
+	}
+
+	// Attenuation
+	if (bAffectedByAttenuation || !bAllowSpatialization)
+	{
+		// Sound and listener are inside
+		SourceInteriorVolume = FMath::Lerp(SourceInteriorVolume, 1.0f, IntVolumeInterp);
+	}
+	else
+	{
+		if (bUsingWorldSettings)
+		{
+			// Sound is outside, listener is inside
+			SourceInteriorVolume = FMath::Lerp(SourceInteriorVolume, ListenerInterior.GetExteriorVolume(), ExtVolumeInterp);
+		}
+		else
+		{
+			// Sound is inside, listener is outside - Use the sound's interior volume multiplied with the listeners exterior volume
+			float SoundInteriorValue = FMath::Lerp(SourceInteriorVolume, Sound.InteriorSettings.GetInteriorVolume(), IntVolumeInterp);
+			float ListenerInteriorValue = FMath::Lerp(SourceInteriorVolume, ListenerInterior.GetExteriorVolume(), ExtVolumeInterp);
+			SourceInteriorVolume = SoundInteriorValue * ListenerInteriorValue;
+		}
+	}
+
+	// Filter
+	if (bAffectedByFilter || !bAllowSpatialization)
+	{
+		// Sound and listener are inside
+		SourceInteriorLPF = FMath::Lerp(SourceInteriorLPF, MAX_FILTER_FREQUENCY, IntLowPassInterp);
+	}
+	else
+	{
+		if (bUsingWorldSettings)
+		{
+			// Sound is outside, listener is inside
+			SourceInteriorLPF = FMath::Lerp(SourceInteriorLPF, ListenerInterior.GetExteriorLPF(), ExtLowPassInterp);
+		}
+		else
+		{
+			// Sound is inside, listener is outside - Current interior LPF is the lesser of the Sound and Listener's LPFs
+			float SoundLPFValue = FMath::Lerp(SourceInteriorLPF, Sound.InteriorSettings.GetInteriorLPF(), IntLowPassInterp);
+			float ListenerLPFValue = FMath::Lerp(SourceInteriorLPF, ListenerInterior.GetExteriorLPF(), ExtLowPassInterp);
+			SourceInteriorLPF = FMath::Min(SoundLPFValue, ListenerLPFValue);
+		}
+	}
+}
+
+
+constexpr TCHAR FProxyVolumeMutator::MutatorBaseName[];
+
+FProxyVolumeMutator::FProxyVolumeMutator()
+{
+	MutatorName = MutatorBaseName;
+}
+
+void FProxyVolumeMutator::UpdatePriority(FAudioProxyMutatorPriorities& Priorities) const
+{
+	if (!Priorities.bFilterPayload || HasPayloadType(Priorities.PayloadType))
+	{
+		int32 DefaultValue = INDEX_NONE;
+		int32& CurrentPriority = Priorities.PriorityMap.FindOrAdd(MutatorName, DefaultValue);
+		if (Priority > CurrentPriority)
+		{
+			CurrentPriority = Priority;
+		}
+	}
+}
+
+bool FProxyVolumeMutator::CheckPriority(const FAudioProxyMutatorPriorities& Priorities) const
+{
+	if (!Priorities.bFilterPayload || HasPayloadType(Priorities.PayloadType))
+	{
+		if (const int32* HighestPriority = Priorities.PriorityMap.Find(MutatorName))
+		{
+			if (Priority == *HighestPriority)
+			{
+				return true;
+			}
+		}
+	}
+
+	return false;
+}
+
+void FProxyVolumeMutator::Apply(FInteriorSettings& InteriorSettings) const
+{
+	InteriorSettings.bIsWorldSettings = false;
+}
+
+bool FProxyVolumeMutator::HasPayloadType(PayloadFlags InType) const
+{
+	return (PayloadType & InType) != PayloadFlags::AGCP_None;
+}
+
+UAudioGameplayVolumeMutator::UAudioGameplayVolumeMutator(const FObjectInitializer& ObjectInitializer)
+	: Super(ObjectInitializer)
+{
+}
+
+void UAudioGameplayVolumeMutator::SetPriority(int32 InPriority)
+{
+	if (Priority != InPriority)
+	{
+		Priority = InPriority;
+		NotifyDataChanged();
+	}
+}
+
+TSharedPtr<FProxyVolumeMutator> UAudioGameplayVolumeMutator::CreateMutator() const
+{
+	TSharedPtr<FProxyVolumeMutator> ProxyMutator = FactoryMutator();
+	if (ProxyMutator.IsValid())
+	{
+		CopyAudioDataToMutatorBase(ProxyMutator);
+	}
+
+	return ProxyMutator;
+}
+
+TSharedPtr<FProxyVolumeMutator> UAudioGameplayVolumeMutator::FactoryMutator() const
+{
+	return TSharedPtr<FProxyVolumeMutator>();
+}
+
+void UAudioGameplayVolumeMutator::NotifyDataChanged() const
+{
+	if (IsActive())
+	{
+		TInlineComponentArray<UAudioGameplayVolumeProxyComponent*> VolumeComponents(GetOwner());
+		if (ensureMsgf(VolumeComponents.Num() == 1, TEXT("Expecting exactly one AudioGameplayVolumeProxyComponent on an actor!")))
+		{
+			if (VolumeComponents[0])
+			{
+				VolumeComponents[0]->OnComponentDataChanged();
+			}
+		}
+	}
+}
+
+void UAudioGameplayVolumeMutator::CopyAudioDataToMutatorBase(TSharedPtr<FProxyVolumeMutator>& Mutator) const
+{
+	check(Mutator.IsValid());
+	Mutator->PayloadType = PayloadType;
+	Mutator->Priority = Priority;
+
+	CopyAudioDataToMutator(Mutator);
+}
diff --git a/Engine/Plugins/AudioGameplayVolume/Source/AudioGameplayVolume/Private/AudioGameplayVolumeProxy.cpp b/Engine/Plugins/AudioGameplayVolume/Source/AudioGameplayVolume/Private/AudioGameplayVolumeProxy.cpp
--- a/Engine/Plugins/AudioGameplayVolume/Source/AudioGameplayVolume/Private/AudioGameplayVolumeProxy.cpp
+++ b/Engine/Plugins/AudioGameplayVolume/Source/AudioGameplayVolume/Private/AudioGameplayVolumeProxy.cpp
@@ -1,7 +1,8 @@
 // Copyright Epic Games, Inc. All Rights Reserved.
 
 #include "AudioGameplayVolumeProxy.h"
-#include "AudioGameplayVolumeProxyMutator.h"
+#include "AudioGameplayVolumeMutator.h"
+#include "AudioGameplayVolumeSubsystem.h"
 #include "AudioGameplayVolumeLogs.h"
 #include "AudioGameplayVolumeComponent.h"
 #include "Interfaces/IAudioGameplayCondition.h"
@@ -32,8 +33,8 @@ void UAudioGameplayVolumeProxy::InitFromComponent(const UAudioGameplayVolumeProx
 	PayloadType = PayloadFlags::AGCP_None;
 	ProxyVolumeMutators.Reset();
 
-	TInlineComponentArray<UAudioGameplayVolumeComponentBase*> Components(Component->GetOwner());
-	for (UAudioGameplayVolumeComponentBase* Comp : Components)
+	TInlineComponentArray<UAudioGameplayVolumeMutator*> Components(Component->GetOwner());
+	for (UAudioGameplayVolumeMutator* Comp : Components)
 	{
 		if (!Comp || !Comp->IsActive())
 		{
diff --git a/Engine/Plugins/AudioGameplayVolume/Source/AudioGameplayVolume/Private/AudioGameplayVolumeProxyMutator.cpp b/Engine/Plugins/AudioGameplayVolume/Source/AudioGameplayVolume/Private/AudioGameplayVolumeProxyMutator.cpp
deleted file mode 100644
diff --git a/Engine/Plugins/AudioGameplayVolume/Source/AudioGameplayVolume/Private/AudioGameplayVolumeSubsystem.cpp b/Engine/Plugins/AudioGameplayVolume/Source/AudioGameplayVolume/Private/AudioGameplayVolumeSubsystem.cpp
--- a/Engine/Plugins/AudioGameplayVolume/Source/AudioGameplayVolume/Private/AudioGameplayVolumeSubsystem.cpp
+++ b/Engine/Plugins/AudioGameplayVolume/Source/AudioGameplayVolume/Private/AudioGameplayVolumeSubsystem.cpp
@@ -4,8 +4,7 @@
 #include "AudioGameplayVolumeLogs.h"
 #include "AudioGameplayVolumeComponent.h"
 #include "AudioGameplayVolumeProxy.h"
-#include "AudioGameplayVolumeProxyMutator.h"
-#include "AudioGameplayFlags.h"
+#include "AudioGameplayVolumeMutator.h"
 #include "AudioDevice.h"
 #include "Misc/CoreMisc.h"
 
@@ -62,6 +61,64 @@ void FAudioGameplayActiveSoundInfo::Update(double ListenerInteriorStartTime)
 	InteriorSettings.UpdateInteriorValues();
 }
 
+void FAudioProxyMutatorSearchObject::SearchVolumes(const TArray<TWeakObjectPtr<UAudioGameplayVolumeProxy>>& ProxyVolumes, FAudioProxyMutatorSearchResult& OutResult)
+{
+	check(IsInAudioThread());
+	SCOPED_NAMED_EVENT(FAudioProxyMutatorSearchObject_SearchVolumes, FColor::Blue);
+
+	OutResult.Reset();
+	FAudioProxyMutatorPriorities MutatorPriorities;
+	MutatorPriorities.PayloadType = PayloadType;
+	MutatorPriorities.bFilterPayload = bFilterPayload;
+
+	for (const TWeakObjectPtr<UAudioGameplayVolumeProxy>& ProxyVolume : ProxyVolumes)
+	{
+		if (!ProxyVolume.IsValid() || ProxyVolume->GetWorldID() != WorldID)
+		{
+			continue;
+		}
+
+		if (bFilterPayload && !ProxyVolume->HasPayloadType(PayloadType))
+		{
+			continue;
+		}
+
+		if (!ProxyVolume->ContainsPosition(Location))
+		{
+			continue;
+		}
+
+		if (bCollectMutators)
+		{
+			// We only need to calculate priorities for mutators if we're collecting them.
+			ProxyVolume->FindMutatorPriority(MutatorPriorities);
+		}
+		OutResult.VolumeSet.Add(ProxyVolume->GetVolumeID());
+	}
+
+	// Use 'world settings' as a starting point
+	if (!bAffectedByLegacySystem && AudioDeviceHandle.IsValid())
+	{
+		AudioDeviceHandle->GetDefaultAudioSettings(WorldID, OutResult.ReverbSettings, OutResult.InteriorSettings);
+	}
+
+	if (bCollectMutators)
+	{
+		for (const TWeakObjectPtr<UAudioGameplayVolumeProxy>& ProxyVolume : ProxyVolumes)
+		{
+			if (!ProxyVolume.IsValid())
+			{
+				continue;
+			}
+
+			if (OutResult.VolumeSet.Contains(ProxyVolume->GetVolumeID()))
+			{
+				ProxyVolume->GatherMutators(MutatorPriorities, OutResult);
+			}
+		}
+	}
+}
+
 void FAudioGameplayVolumeProxyInfo::Update(const TArray<FAudioGameplayVolumeListener>& VolumeListeners, FAudioGameplayProxyUpdateResult& OutResult)
 {
 	check(IsInAudioThread());
diff --git a/Engine/Plugins/AudioGameplayVolume/Source/AudioGameplayVolume/Public/AttenuationVolumeComponent.h b/Engine/Plugins/AudioGameplayVolume/Source/AudioGameplayVolume/Public/AttenuationVolumeComponent.h
--- a/Engine/Plugins/AudioGameplayVolume/Source/AudioGameplayVolume/Public/AttenuationVolumeComponent.h
+++ b/Engine/Plugins/AudioGameplayVolume/Source/AudioGameplayVolume/Public/AttenuationVolumeComponent.h
@@ -2,8 +2,7 @@
 
 #pragma once
 
-#include "AudioGameplayVolumeComponent.h"
-#include "AudioGameplayVolumeProxyMutator.h"
+#include "AudioGameplayVolumeMutator.h"
 #include "AttenuationVolumeComponent.generated.h"
 
 /**
@@ -33,7 +32,7 @@ protected:
  *  UAttenuationVolumeComponent - Audio Gameplay Volume component for occlusion settings (volume attenuation)
  */
 UCLASS(Blueprintable, Config = Game, ClassGroup = ("AudioGameplayVolume"), meta = (BlueprintSpawnableComponent, DisplayName = "Attenuation"))
-class AUDIOGAMEPLAYVOLUME_API UAttenuationVolumeComponent : public UAudioGameplayVolumeComponentBase
+class AUDIOGAMEPLAYVOLUME_API UAttenuationVolumeComponent : public UAudioGameplayVolumeMutator
 {
 	GENERATED_UCLASS_BODY()
 
@@ -55,10 +54,10 @@ public:
 
 private:
 
-	//~ Begin UAudioGameplayVolumeComponentBase interface
+	//~ Begin UAudioGameplayVolumeMutator interface
 	virtual TSharedPtr<FProxyVolumeMutator> FactoryMutator() const override;
 	virtual void CopyAudioDataToMutator(TSharedPtr<FProxyVolumeMutator>& Mutator) const override;
-	//~ End UAudioGameplayVolumeComponentBase interface
+	//~ End UAudioGameplayVolumeMutator interface
 
 	// The desired volume of sounds outside the volume when the player is inside the volume
 	UPROPERTY(EditAnywhere, BlueprintReadWrite, Category = "VolumeAttenuation", meta = (AllowPrivateAccess = "true"))
diff --git a/Engine/Plugins/AudioGameplayVolume/Source/AudioGameplayVolume/Public/AudioGameplayVolumeComponent.h b/Engine/Plugins/AudioGameplayVolume/Source/AudioGameplayVolume/Public/AudioGameplayVolumeComponent.h
--- a/Engine/Plugins/AudioGameplayVolume/Source/AudioGameplayVolume/Public/AudioGameplayVolumeComponent.h
+++ b/Engine/Plugins/AudioGameplayVolume/Source/AudioGameplayVolume/Public/AudioGameplayVolumeComponent.h
@@ -6,56 +6,14 @@
 #include "AudioGameplayVolumeComponent.generated.h"
 
 // Forward Declarations 
-class FProxyVolumeMutator;
 class UAudioGameplayVolumeProxy;
 class UAudioGameplayVolumeSubsystem;
 
-/**
- *  UAudioGameplayVolumeComponentBase - Base component for use with audio gameplay volumes
- */
-UCLASS(Abstract, HideDropdown)
-class AUDIOGAMEPLAYVOLUME_API UAudioGameplayVolumeComponentBase : public UAudioGameplayComponent
-{
-	GENERATED_UCLASS_BODY()
-
-public:
-
-	virtual ~UAudioGameplayVolumeComponentBase() = default;
-
-	UFUNCTION(BlueprintCallable, Category = "AudioGameplay")
-	void SetPriority(int32 InPriority);
-
-	int32 GetPriority() const { return Priority; }
-
-	/** Create and fill the appropriate proxy mutator for this component */
-	virtual TSharedPtr<FProxyVolumeMutator> CreateMutator() const final;
-
-protected:
-
-	/** Create this component's type of mutator */
-	virtual TSharedPtr<FProxyVolumeMutator> FactoryMutator() const;
-
-	/** Override in child classes to copy additional data needed to mutators */
-	virtual void CopyAudioDataToMutator(TSharedPtr<FProxyVolumeMutator>& Mutator) const {}
-
-	/** Notify our parent volume our proxy may need updating */
-	void NotifyDataChanged() const;
-
-	// The priority of this component.  In the case of overlapping volumes or multiple affecting components, the highest priority is chosen.
-	UPROPERTY(EditAnywhere, BlueprintReadWrite, Category = "AudioGameplay", Meta = (AllowPrivateAccess = "true"))
-	int32 Priority = 0;
-
-private:
-
-	/** Called for you during mutator creation. See CopyAudioDataToMutator for adding data to derived classes */
-	void CopyAudioDataToMutatorBase(TSharedPtr<FProxyVolumeMutator>& Mutator) const;
-};
-
 DECLARE_DYNAMIC_MULTICAST_DELEGATE(FOnAudioGameplayVolumeProxyStateChange);
 
 /**
  *  UAudioGameplayVolumeProxyComponent - Component used to drive interaction with AudioGameplayVolumeSubsystem.
- *   NOTE: Do not inherit from this class, use UAudioGameplayVolumeComponentBase to create extendable functionality
+ *   NOTE: Do not inherit from this class, use UAudioGameplayVolumeMutator to create extendable functionality
  */
 UCLASS(Config = Game, ClassGroup = ("AudioGameplay"), meta = (BlueprintSpawnableComponent, IsBlueprintBase = false, DisplayName = "Volume Proxy"))
 class AUDIOGAMEPLAYVOLUME_API UAudioGameplayVolumeProxyComponent final : public UAudioGameplayComponent
diff --git a/Engine/Plugins/AudioGameplayVolume/Source/AudioGameplayVolume/Public/AudioGameplayVolumeMutator.h b/Engine/Plugins/AudioGameplayVolume/Source/AudioGameplayVolume/Public/AudioGameplayVolumeMutator.h
new file mode 100644
--- /dev/null
+++ b/Engine/Plugins/AudioGameplayVolume/Source/AudioGameplayVolume/Public/AudioGameplayVolumeMutator.h
@@ -0,0 +1,132 @@
+// Copyright Epic Games, Inc. All Rights Reserved.
+
+#pragma once
+
+#include "AudioDefines.h"
+#include "Sound/AudioVolume.h"
+#include "AudioGameplayFlags.h"
+#include "Templates/SharedPointer.h"
+#include "AudioGameplayVolumeComponent.h"
+#include "AudioGameplayVolumeMutator.generated.h"
+
+// Forward Declarations 
+class FAudioGameplayVolumeListener;
+struct FAudioGameplayActiveSoundInfo;
+
+/**
+ *  FAudioProxyActiveSoundParams - Helper struct for collecting info about the active sound from affecting proxy mutators
+ */
+struct FAudioProxyActiveSoundParams
+{
+	FAudioProxyActiveSoundParams() = delete;
+	FAudioProxyActiveSoundParams(const FAudioGameplayActiveSoundInfo& SoundInfo, const FAudioGameplayVolumeListener& InListener);
+
+	float SourceInteriorVolume = 1.f;
+
+	float SourceInteriorLPF = MAX_FILTER_FREQUENCY;
+
+	bool bAllowSpatialization = false;
+	bool bUsingWorldSettings = false;
+	bool bListenerInVolume = false;
+
+	bool bAffectedByAttenuation = false;
+	bool bAffectedByFilter = false;
+
+	const FAudioGameplayVolumeListener& Listener;
+	const FAudioGameplayActiveSoundInfo& Sound;
+
+	TArray<FSoundSubmixSendInfo> SoundSubmixSends;
+
+	void UpdateInteriorValues();
+};
+
+/**
+ *  FAudioProxyMutatorPriorities - Used for finding the highest priority mutators on a proxy
+ */
+struct FAudioProxyMutatorPriorities
+{
+	using PayloadFlags = AudioGameplay::EComponentPayload;
+
+	TMap<FName, int32> PriorityMap;
+	PayloadFlags PayloadType = PayloadFlags::AGCP_None;
+	bool bFilterPayload = true;
+};
+
+/**
+ *  FProxyVolumeMutator - An audio thread representation of the payload for an AudioGameplayVolumeComponent.
+ */
+class AUDIOGAMEPLAYVOLUME_API FProxyVolumeMutator : public TSharedFromThis<FProxyVolumeMutator>
+{
+public:
+
+	FProxyVolumeMutator();
+	virtual ~FProxyVolumeMutator() = default;
+
+	using PayloadFlags = AudioGameplay::EComponentPayload;
+
+	virtual void UpdatePriority(FAudioProxyMutatorPriorities& Priorities) const;
+	virtual bool CheckPriority(const FAudioProxyMutatorPriorities& Priorities) const;
+
+	virtual void Apply(FInteriorSettings& InteriorSettings) const;
+	virtual void Apply(FAudioProxyActiveSoundParams& Params) const {}
+	virtual void Apply(FAudioGameplayVolumeListener& Listener) const {}
+
+	virtual void Remove(FAudioProxyActiveSoundParams& Params) const {}
+	virtual void Remove(FAudioGameplayVolumeListener& Listener) const {}
+
+	bool HasPayloadType(PayloadFlags InType) const;
+
+	int32 Priority = INDEX_NONE;
+	uint32 VolumeID = INDEX_NONE;
+	uint32 WorldID = INDEX_NONE;
+
+	FName MutatorName;
+
+	PayloadFlags PayloadType;
+
+protected:
+
+	constexpr static const TCHAR MutatorBaseName[] = TEXT("MutatorBase");
+};
+
+// #todo_rpm comments - more complex; current not descriptive
+/**
+ *  UAudioGameplayVolumeMutator - Base component for use with audio gameplay volumes
+ */
+UCLASS(Abstract, HideDropdown)
+class AUDIOGAMEPLAYVOLUME_API UAudioGameplayVolumeMutator : public UAudioGameplayComponent
+{
+	GENERATED_UCLASS_BODY()
+
+public:
+
+	virtual ~UAudioGameplayVolumeMutator() = default;
+
+	UFUNCTION(BlueprintCallable, Category = "AudioGameplay")
+	void SetPriority(int32 InPriority);
+
+	int32 GetPriority() const { return Priority; }
+
+	/** Create and fill the appropriate proxy mutator for this component */
+	virtual TSharedPtr<FProxyVolumeMutator> CreateMutator() const final;
+
+protected:
+
+	/** Create this component's type of mutator */
+	virtual TSharedPtr<FProxyVolumeMutator> FactoryMutator() const;
+
+	/** Override in child classes to copy additional data needed to mutators */
+	virtual void CopyAudioDataToMutator(TSharedPtr<FProxyVolumeMutator>& Mutator) const {}
+
+	/** Notify our parent volume our proxy may need updating */
+	void NotifyDataChanged() const;
+
+	// The priority of this component.  In the case of overlapping volumes or multiple affecting components, the highest priority is chosen.
+	UPROPERTY(EditAnywhere, BlueprintReadWrite, Category = "AudioGameplay", Meta = (AllowPrivateAccess = "true"))
+	int32 Priority = 0;
+
+private:
+
+	/** Called for you during mutator creation. See CopyAudioDataToMutator for adding data to derived classes */
+	void CopyAudioDataToMutatorBase(TSharedPtr<FProxyVolumeMutator>& Mutator) const;
+};
diff --git a/Engine/Plugins/AudioGameplayVolume/Source/AudioGameplayVolume/Public/AudioGameplayVolumeProxyMutator.h b/Engine/Plugins/AudioGameplayVolume/Source/AudioGameplayVolume/Public/AudioGameplayVolumeProxyMutator.h
--- a/Engine/Plugins/AudioGameplayVolume/Source/AudioGameplayVolume/Public/AudioGameplayVolumeProxyMutator.h
+++ b/Engine/Plugins/AudioGameplayVolume/Source/AudioGameplayVolume/Public/AudioGameplayVolumeProxyMutator.h
@@ -1,129 +1,22 @@
 // Copyright Epic Games, Inc. All Rights Reserved.
 
-#pragma once
-
-#include "Sound/AudioVolume.h"
-#include "AudioDeviceManager.h"
-#include "AudioGameplayFlags.h"
-#include "Templates/SharedPointer.h"
-
-// Forward Declarations 
-class FProxyVolumeMutator;
-class UAudioGameplayVolumeProxy;
-class FAudioGameplayVolumeListener;
-struct FAudioGameplayActiveSoundInfo;
-
-/**
- *  FAudioProxyMutatorSearchResult - Results from a audio proxy mutator search (see below).
- */
-struct FAudioProxyMutatorSearchResult
-{
-	TSet<uint32> VolumeSet;
-	TArray<TSharedPtr<FProxyVolumeMutator>> MatchingMutators;
-	FReverbSettings ReverbSettings;
-	FInteriorSettings InteriorSettings;
-
-	void Reset()
-	{
-		VolumeSet.Reset();
-		MatchingMutators.Empty();
-		ReverbSettings = FReverbSettings();
-		InteriorSettings = FInteriorSettings();
-	}
-};
-
-/**
- *  FAudioProxyMutatorSearchObject - Used for searching through proxy volumes to find relevant proxy mutators
- */
-struct FAudioProxyMutatorSearchObject
-{
-	using PayloadFlags = AudioGameplay::EComponentPayload;
-
-	// Search parameters
-	uint32 WorldID = INDEX_NONE;
-	FVector Location = FVector::ZeroVector;
-	PayloadFlags PayloadType = PayloadFlags::AGCP_None;
-	FAudioDeviceHandle AudioDeviceHandle;
-	bool bAffectedByLegacySystem = false;
-	bool bFilterPayload = true;
-	bool bCollectMutators = true;
-	bool bGetDefaultAudioSettings = true;
-
-	void SearchVolumes(const TArray<TWeakObjectPtr<UAudioGameplayVolumeProxy>>& ProxyVolumes, FAudioProxyMutatorSearchResult& OutResult);
-};
-
-/**
- *  FAudioProxyMutatorPriorities - Used for finding the highest priority mutators on a proxy
- */
-struct FAudioProxyMutatorPriorities
-{
-	using PayloadFlags = AudioGameplay::EComponentPayload;
-
-	TMap<FName, int32> PriorityMap;
-	PayloadFlags PayloadType = PayloadFlags::AGCP_None;
-	bool bFilterPayload = true;
-};
-
-/**
- *  FAudioProxyActiveSoundParams - Helper struct for collecting info about the active sound from affecting proxy mutators
- */
-struct FAudioProxyActiveSoundParams
-{
-	FAudioProxyActiveSoundParams() = delete;
-	FAudioProxyActiveSoundParams(const FAudioGameplayActiveSoundInfo& SoundInfo, const FAudioGameplayVolumeListener& InListener);
-
-	float SourceInteriorVolume = 1.f;
-
-	float SourceInteriorLPF = MAX_FILTER_FREQUENCY;
-
-	bool bAllowSpatialization = false;
-	bool bUsingWorldSettings = false;
-	bool bListenerInVolume = false;
-
-	bool bAffectedByAttenuation = false;
-	bool bAffectedByFilter = false;
-
-	const FAudioGameplayVolumeListener& Listener;
-	const FAudioGameplayActiveSoundInfo& Sound;
-
-	TArray<FSoundSubmixSendInfo> SoundSubmixSends;
-
-	void UpdateInteriorValues();
-};
-
 /**
- *  FProxyVolumeMutator - An audio thread representation of the payload for an AudioGameplayVolumeComponent.
+ * DEPRECATED AudioGameplayVolumeProxyMutator.h.  Use the following files instead:
+	 AudioGameplayVolumeSubsystem.h for FAudioProxyMutatorSearchResult
+	 AudioGameplayVolumeSubsystem.h for FAudioProxyMutatorSearchObject
+	 AudioGameplayVolumeMutator.h	for FAudioProxyMutatorPriorities
+	 AudioGameplayVolumeMutator.h	for FAudioProxyActiveSoundParams
+	 AudioGameplayVolumeMutator.h	for FProxyVolumeMutator
  */
-class AUDIOGAMEPLAYVOLUME_API FProxyVolumeMutator : public TSharedFromThis<FProxyVolumeMutator>
-{
-public:
 
-	FProxyVolumeMutator();
-	virtual ~FProxyVolumeMutator() = default;
-
-	using PayloadFlags = AudioGameplay::EComponentPayload;
-
-	virtual void UpdatePriority(FAudioProxyMutatorPriorities& Priorities) const;
-	virtual bool CheckPriority(const FAudioProxyMutatorPriorities& Priorities) const;
-
-	virtual void Apply(FInteriorSettings& InteriorSettings) const;
-	virtual void Apply(FAudioProxyActiveSoundParams& Params) const {}
-	virtual void Apply(FAudioGameplayVolumeListener& Listener) const {}
-
-	virtual void Remove(FAudioProxyActiveSoundParams& Params) const {}
-	virtual void Remove(FAudioGameplayVolumeListener& Listener) const {}
-
-	bool HasPayloadType(PayloadFlags InType) const;
-
-	int32 Priority = INDEX_NONE;
-	uint32 VolumeID = INDEX_NONE;
-	uint32 WorldID = INDEX_NONE;
-
-	FName MutatorName;
-
-	PayloadFlags PayloadType;
+#pragma once
 
-protected:
+#ifdef _MSC_VER
+#pragma message(__FILE__"(15): warning: use AudioGameplayVolumeSubsystem.h and AudioGameplayVolumeMutator.h instead of AudioGameplayVolumeProxyMutator.h")
+#else
+#pragma message("#include AudioGameplayVolumeSubsystem.h and AudioGameplayVolumeMutator.h instead of AudioGameplayVolumeProxyMutator.h")
+#endif
 
-	constexpr static const TCHAR MutatorBaseName[] = TEXT("MutatorBase");
-};
+ // Include the new file so that the project still compiles, but has warnings
+#include "AudioGameplayVolumeSubsystem.h"
+#include "AudioGameplayVolumeMutator.h"
diff --git a/Engine/Plugins/AudioGameplayVolume/Source/AudioGameplayVolume/Public/AudioGameplayVolumeSubsystem.h b/Engine/Plugins/AudioGameplayVolume/Source/AudioGameplayVolume/Public/AudioGameplayVolumeSubsystem.h
--- a/Engine/Plugins/AudioGameplayVolume/Source/AudioGameplayVolume/Public/AudioGameplayVolumeSubsystem.h
+++ b/Engine/Plugins/AudioGameplayVolume/Source/AudioGameplayVolume/Public/AudioGameplayVolumeSubsystem.h
@@ -4,7 +4,9 @@
 
 #include "ActiveSoundUpdateInterface.h"
 #include "Subsystems/AudioEngineSubsystem.h"
+#include "AudioGameplayFlags.h"
 #include "AudioGameplayVolumeListener.h"
+#include "Templates/SharedPointer.h"
 #include "AudioGameplayVolumeSubsystem.generated.h"
 
 // Forward Declarations 
@@ -43,6 +45,45 @@ struct FAudioGameplayProxyUpdateResult
 	bool bForceUpdate = false;
 };
 
+/**
+ *  FAudioProxyMutatorSearchResult - Results from a audio proxy mutator search (see below).
+ */
+struct FAudioProxyMutatorSearchResult
+{
+	TSet<uint32> VolumeSet;
+	TArray<TSharedPtr<FProxyVolumeMutator>> MatchingMutators;
+	FReverbSettings ReverbSettings;
+	FInteriorSettings InteriorSettings;
+
+	void Reset()
+	{
+		VolumeSet.Reset();
+		MatchingMutators.Empty();
+		ReverbSettings = FReverbSettings();
+		InteriorSettings = FInteriorSettings();
+	}
+};
+
+/**
+ *  FAudioProxyMutatorSearchObject - Used for searching through proxy volumes to find relevant proxy mutators
+ */
+struct FAudioProxyMutatorSearchObject
+{
+	using PayloadFlags = AudioGameplay::EComponentPayload;
+
+	// Search parameters
+	uint32 WorldID = INDEX_NONE;
+	FVector Location = FVector::ZeroVector;
+	PayloadFlags PayloadType = PayloadFlags::AGCP_None;
+	FAudioDeviceHandle AudioDeviceHandle;
+	bool bAffectedByLegacySystem = false;
+	bool bFilterPayload = true;
+	bool bCollectMutators = true;
+	bool bGetDefaultAudioSettings = true;
+
+	void SearchVolumes(const TArray<TWeakObjectPtr<UAudioGameplayVolumeProxy>>& ProxyVolumes, FAudioProxyMutatorSearchResult& OutResult);
+};
+
 /**
  *  FAudioGameplayVolumeProxyInfo - Holds information relating to which volumes our listeners are inside of
  */
diff --git a/Engine/Plugins/AudioGameplayVolume/Source/AudioGameplayVolume/Public/FilterVolumeComponent.h b/Engine/Plugins/AudioGameplayVolume/Source/AudioGameplayVolume/Public/FilterVolumeComponent.h
--- a/Engine/Plugins/AudioGameplayVolume/Source/AudioGameplayVolume/Public/FilterVolumeComponent.h
+++ b/Engine/Plugins/AudioGameplayVolume/Source/AudioGameplayVolume/Public/FilterVolumeComponent.h
@@ -2,9 +2,7 @@
 
 #pragma once
 
-#include "AudioGameplayVolumeComponent.h"
-#include "AudioGameplayVolumeProxyMutator.h"
-#include "AudioDefines.h"
+#include "AudioGameplayVolumeMutator.h"
 #include "FilterVolumeComponent.generated.h"
 
 /**
@@ -34,7 +32,7 @@ protected:
  *  UFilterVolumeComponent - Audio Gameplay Volume component for occlusion settings (volume filter)
  */
 UCLASS(Blueprintable, Config = Game, ClassGroup = ("AudioGameplayVolume"), meta = (BlueprintSpawnableComponent, DisplayName = "Filter"))
-class AUDIOGAMEPLAYVOLUME_API UFilterVolumeComponent : public UAudioGameplayVolumeComponentBase
+class AUDIOGAMEPLAYVOLUME_API UFilterVolumeComponent : public UAudioGameplayVolumeMutator
 {
 	GENERATED_UCLASS_BODY()
 
@@ -56,10 +54,10 @@ public:
 
 private:
 
-	//~ Begin UAudioGameplayVolumeComponentBase interface
+	//~ Begin UAudioGameplayVolumeMutator interface
 	virtual TSharedPtr<FProxyVolumeMutator> FactoryMutator() const override;
 	virtual void CopyAudioDataToMutator(TSharedPtr<FProxyVolumeMutator>& Mutator) const override;
-	//~ End UAudioGameplayVolumeComponentBase interface
+	//~ End UAudioGameplayVolumeMutator interface
 
 	// The desired LPF frequency cutoff (in hertz) of sounds outside the volume when the player is inside the volume
 	UPROPERTY(EditAnywhere, BlueprintReadWrite, Category = "VolumeFilter", meta = (AllowPrivateAccess = "true"))
diff --git a/Engine/Plugins/AudioGameplayVolume/Source/AudioGameplayVolume/Public/ReverbVolumeComponent.h b/Engine/Plugins/AudioGameplayVolume/Source/AudioGameplayVolume/Public/ReverbVolumeComponent.h
--- a/Engine/Plugins/AudioGameplayVolume/Source/AudioGameplayVolume/Public/ReverbVolumeComponent.h
+++ b/Engine/Plugins/AudioGameplayVolume/Source/AudioGameplayVolume/Public/ReverbVolumeComponent.h
@@ -2,8 +2,7 @@
 
 #pragma once
 
-#include "AudioGameplayVolumeComponent.h"
-#include "AudioGameplayVolumeProxyMutator.h"
+#include "AudioGameplayVolumeMutator.h"
 #include "Sound/ReverbSettings.h"
 #include "ReverbVolumeComponent.generated.h"
 
@@ -31,7 +30,7 @@ protected:
  *  UReverbVolumeComponent - Audio Gameplay Volume component for reverb settings
  */
 UCLASS(Blueprintable, Config = Game, ClassGroup = ("AudioGameplayVolume"), meta = (BlueprintSpawnableComponent, DisplayName = "Reverb"))
-class AUDIOGAMEPLAYVOLUME_API UReverbVolumeComponent : public UAudioGameplayVolumeComponentBase
+class AUDIOGAMEPLAYVOLUME_API UReverbVolumeComponent : public UAudioGameplayVolumeMutator
 {
 	GENERATED_UCLASS_BODY()
 
@@ -46,10 +45,10 @@ public:
 
 private:
 
-	//~ Begin UAudioGameplayVolumeComponentBase interface
+	//~ Begin UAudioGameplayVolumeMutator interface
 	virtual TSharedPtr<FProxyVolumeMutator> FactoryMutator() const override;
 	virtual void CopyAudioDataToMutator(TSharedPtr<FProxyVolumeMutator>& Mutator) const override;
-	//~ End UAudioGameplayVolumeComponentBase interface
+	//~ End UAudioGameplayVolumeMutator interface
 
 	/** Reverb settings to use with this component */
 	UPROPERTY(EditAnywhere, BlueprintReadOnly, Category = Reverb, meta = (AllowPrivateAccess = "true"))
diff --git a/Engine/Plugins/AudioGameplayVolume/Source/AudioGameplayVolume/Public/SubmixOverrideVolumeComponent.h b/Engine/Plugins/AudioGameplayVolume/Source/AudioGameplayVolume/Public/SubmixOverrideVolumeComponent.h
--- a/Engine/Plugins/AudioGameplayVolume/Source/AudioGameplayVolume/Public/SubmixOverrideVolumeComponent.h
+++ b/Engine/Plugins/AudioGameplayVolume/Source/AudioGameplayVolume/Public/SubmixOverrideVolumeComponent.h
@@ -2,9 +2,7 @@
 
 #pragma once
 
-#include "AudioGameplayVolumeComponent.h"
-#include "AudioGameplayVolumeProxyMutator.h"
-#include "Sound/AudioVolume.h"
+#include "AudioGameplayVolumeMutator.h"
 #include "SubmixOverrideVolumeComponent.generated.h"
 
 /**
@@ -31,7 +29,7 @@ protected:
  *  USubmixOverrideVolumeComponent - Audio Gameplay Volume component for submix effect chain overrides
  */
 UCLASS(Blueprintable, Config = Game, ClassGroup = ("AudioGameplayVolume"), meta = (BlueprintSpawnableComponent, DisplayName = "Submix Override"))
-class AUDIOGAMEPLAYVOLUME_API USubmixOverrideVolumeComponent : public UAudioGameplayVolumeComponentBase
+class AUDIOGAMEPLAYVOLUME_API USubmixOverrideVolumeComponent : public UAudioGameplayVolumeMutator
 {
 	GENERATED_UCLASS_BODY()
 
@@ -46,10 +44,10 @@ public:
 
 private:
 
-	//~ Begin UAudioGameplayVolumeComponentBase interface
+	//~ Begin UAudioGameplayVolumeMutator interface
 	virtual TSharedPtr<FProxyVolumeMutator> FactoryMutator() const override;
 	virtual void CopyAudioDataToMutator(TSharedPtr<FProxyVolumeMutator>& Mutator) const override;
-	//~ End UAudioGameplayVolumeComponentBase interface
+	//~ End UAudioGameplayVolumeMutator interface
 
 	/** Submix effect chain override settings. Will override the effect chains on the given submixes */
 	UPROPERTY(EditAnywhere, BlueprintReadOnly, Category = Submixes, meta = (AllowPrivateAccess = "true"))
diff --git a/Engine/Plugins/AudioGameplayVolume/Source/AudioGameplayVolume/Public/SubmixSendVolumeComponent.h b/Engine/Plugins/AudioGameplayVolume/Source/AudioGameplayVolume/Public/SubmixSendVolumeComponent.h
--- a/Engine/Plugins/AudioGameplayVolume/Source/AudioGameplayVolume/Public/SubmixSendVolumeComponent.h
+++ b/Engine/Plugins/AudioGameplayVolume/Source/AudioGameplayVolume/Public/SubmixSendVolumeComponent.h
@@ -2,9 +2,7 @@
 
 #pragma once
 
-#include "AudioGameplayVolumeComponent.h"
-#include "AudioGameplayVolumeProxyMutator.h"
-#include "Sound/AudioVolume.h"
+#include "AudioGameplayVolumeMutator.h"
 #include "SubmixSendVolumeComponent.generated.h"
 
 /**
@@ -30,7 +28,7 @@ protected:
  *  USubmixSendVolumeComponent - Audio Gameplay Volume component for submix sends
  */
 UCLASS(Blueprintable, Config = Game, ClassGroup = ("AudioGameplayVolume"), meta = (BlueprintSpawnableComponent, DisplayName = "Submix Send"))
-class AUDIOGAMEPLAYVOLUME_API USubmixSendVolumeComponent : public UAudioGameplayVolumeComponentBase
+class AUDIOGAMEPLAYVOLUME_API USubmixSendVolumeComponent : public UAudioGameplayVolumeMutator
 {
 	GENERATED_UCLASS_BODY()
 
@@ -45,10 +43,10 @@ public:
 
 private:
 
-	//~ Begin UAudioGameplayVolumeComponentBase interface
+	//~ Begin UAudioGameplayVolumeMutator interface
 	virtual TSharedPtr<FProxyVolumeMutator> FactoryMutator() const override;
 	virtual void CopyAudioDataToMutator(TSharedPtr<FProxyVolumeMutator>& Mutator) const override;
-	//~ End UAudioGameplayVolumeComponentBase interface
+	//~ End UAudioGameplayVolumeMutator interface
 
 	/** Submix send settings to use for this component. Allows audio to dynamically send to submixes based on source and listener locations (relative to parent volume.) */
 	UPROPERTY(EditAnywhere, BlueprintReadOnly, Category = Submixes, meta = (AllowPrivateAccess = "true"))
