Updating missing comments from header file (prior submit 20826820)

#rnx

diff --git a/Engine/Plugins/AudioGameplayVolume/Source/AudioGameplayVolume/Public/AudioGameplayVolumeMutator.h b/Engine/Plugins/AudioGameplayVolume/Source/AudioGameplayVolume/Public/AudioGameplayVolumeMutator.h
--- a/Engine/Plugins/AudioGameplayVolume/Source/AudioGameplayVolume/Public/AudioGameplayVolumeMutator.h
+++ b/Engine/Plugins/AudioGameplayVolume/Source/AudioGameplayVolume/Public/AudioGameplayVolumeMutator.h
@@ -89,9 +89,10 @@ protected:
 	constexpr static const TCHAR MutatorBaseName[] = TEXT("MutatorBase");
 };
 
-// #todo_rpm comments - more complex; current not descriptive
 /**
- *  UAudioGameplayVolumeMutator - Base component for use with audio gameplay volumes
+ *  UAudioGameplayVolumeMutator - These components are used for more complex interactions with AudioGameplayVolumes.
+ *		Currently, components inheriting this base can affect interior settings as well as active sounds or the audio listener(s) inside the volume.
+ *		See also: FilterVolumeComponent, AttenuationVolumeComponent, SubmixSendComponent, SubmixOverrideVolumeComponent, and ReverbVolumeComponent
  */
 UCLASS(Abstract, HideDropdown)
 class AUDIOGAMEPLAYVOLUME_API UAudioGameplayVolumeMutator : public UAudioGameplayComponent
