AudioGameplayVolumes - Prevent Listener Enter and Exit messages from propagating to components that are not enabled.

diff --git a/Engine/Plugins/AudioGameplayVolume/Source/AudioGameplayVolume/Private/AudioGameplayVolumeComponent.cpp b/Engine/Plugins/AudioGameplayVolume/Source/AudioGameplayVolume/Private/AudioGameplayVolumeComponent.cpp
--- a/Engine/Plugins/AudioGameplayVolume/Source/AudioGameplayVolume/Private/AudioGameplayVolumeComponent.cpp
+++ b/Engine/Plugins/AudioGameplayVolume/Source/AudioGameplayVolume/Private/AudioGameplayVolumeComponent.cpp
@@ -36,7 +36,12 @@ void UAudioGameplayVolumeProxyComponent::EnterProxy() const
 	TInlineComponentArray<UActorComponent*> ActorComponents(GetOwner());
 	for (UActorComponent* ActorComponent : ActorComponents)
 	{
-		if (ActorComponent && ActorComponent->Implements<UAudioGameplayVolumeInteraction>())
+		if (!ActorComponent || !ActorComponent->IsActive())
+		{
+			continue;
+		}
+
+		if (ActorComponent->Implements<UAudioGameplayVolumeInteraction>())
 		{
 			IAudioGameplayVolumeInteraction::Execute_OnListenerEnter(ActorComponent);
 		}
@@ -50,7 +55,12 @@ void UAudioGameplayVolumeProxyComponent::ExitProxy() const
 	TInlineComponentArray<UActorComponent*> ActorComponents(GetOwner());
 	for (UActorComponent* ActorComponent : ActorComponents)
 	{
-		if (ActorComponent && ActorComponent->Implements<UAudioGameplayVolumeInteraction>())
+		if (!ActorComponent || !ActorComponent->IsActive())
+		{
+			continue;
+		}
+
+		if (ActorComponent->Implements<UAudioGameplayVolumeInteraction>())
 		{
 			IAudioGameplayVolumeInteraction::Execute_OnListenerExit(ActorComponent);
 		}
