[AGV] Fix issues concerning an uninitialized listener at the origin

diff --git a/Engine/Plugins/AudioGameplayVolume/Source/AudioGameplayVolume/Private/AudioGameplayVolumeSubsystem.cpp b/Engine/Plugins/AudioGameplayVolume/Source/AudioGameplayVolume/Private/AudioGameplayVolumeSubsystem.cpp
--- a/Engine/Plugins/AudioGameplayVolume/Source/AudioGameplayVolume/Private/AudioGameplayVolumeSubsystem.cpp
+++ b/Engine/Plugins/AudioGameplayVolume/Source/AudioGameplayVolume/Private/AudioGameplayVolumeSubsystem.cpp
@@ -477,6 +477,7 @@ void UAudioGameplayVolumeSubsystem::UpdateFromListeners()
 
 	FAudioDeviceHandle DeviceHandle = GetAudioDeviceHandle();
 	check(DeviceHandle.IsValid());
+	const uint32 AudioDeviceID = DeviceHandle.GetDeviceID();
 
 	constexpr bool bAllowShrink = false;
 	const int32 ListenerCount = DeviceHandle->GetListeners().Num();
@@ -490,6 +491,7 @@ void UAudioGameplayVolumeSubsystem::UpdateFromListeners()
 
 	FAudioProxyMutatorSearchResult Result;
 	TSet<uint32> TempVolumeSet;
+	FTransform ListenerTransform;
 
 	// Grabbing the listeners directly here should be removed when possible - done out of necessity due to the legacy audio volume system
 	const TArray<FListener>& AudioListeners = DeviceHandle->GetListeners();
@@ -497,38 +499,48 @@ void UAudioGameplayVolumeSubsystem::UpdateFromListeners()
 	// Update our audio gameplay volume listeners
 	for (int32 i = 0; i < ListenerCount; ++i)
 	{
-		// Fill location and worldID
-		if (DeviceHandle->GetListenerWorldID(i, ProxySearch.WorldID) && DeviceHandle->GetListenerPosition(i, ProxySearch.Location, bAllowAttenuationOverride))
+		const bool bValidAudioListener = DeviceHandle->GetListenerTransform(i, ListenerTransform);
+		if (bValidAudioListener && !ListenerTransform.Equals(FTransform::Identity))
 		{
-			// Find only the proxy volumes we're inside of, regardless of payload type
-			ProxySearch.bFilterPayload = false;
-			ProxySearch.bCollectMutators = false;
-			ProxySearch.bGetDefaultAudioSettings = false;
-			ProxySearch.SearchVolumes(ProxyVolumes, Result);
-
-			// Hold on to these
-			Swap(TempVolumeSet, Result.VolumeSet);
-
-			// Second search - this time for mutators
-			ProxySearch.bFilterPayload = true;
-			ProxySearch.bCollectMutators = true;
-			ProxySearch.bGetDefaultAudioSettings = true;
-			ProxySearch.SearchVolumes(ProxyVolumes, Result);
-
-			if (i < AudioListeners.Num())
+			// Fill location and worldID
+			if (DeviceHandle->GetListenerWorldID(i, ProxySearch.WorldID) && DeviceHandle->GetListenerPosition(i, ProxySearch.Location, bAllowAttenuationOverride))
 			{
-				AGVListeners[i].SetAffectedByLegacySystem(AudioListeners[i].AudioVolumeID != 0);
-			}
-
-			// Reassign the set of all proxy volumes (regardless of payload type)
-			Swap(TempVolumeSet, Result.VolumeSet);
-			AGVListeners[i].Update(Result, ProxySearch.Location, DeviceHandle.GetDeviceID());
+				// Find only the proxy volumes we're inside of, regardless of payload type
+				ProxySearch.bFilterPayload = false;
+				ProxySearch.bCollectMutators = false;
+				ProxySearch.bGetDefaultAudioSettings = false;
+				ProxySearch.SearchVolumes(ProxyVolumes, Result);
+
+				// Hold on to these
+				Swap(TempVolumeSet, Result.VolumeSet);
+
+				// Second search - this time for mutators
+				ProxySearch.bFilterPayload = true;
+				ProxySearch.bCollectMutators = true;
+				ProxySearch.bGetDefaultAudioSettings = true;
+				ProxySearch.SearchVolumes(ProxyVolumes, Result);
+
+				if (i < AudioListeners.Num())
+				{
+					AGVListeners[i].SetAffectedByLegacySystem(AudioListeners[i].AudioVolumeID != 0);
+				}
+
+				// Reassign the set of all proxy volumes (regardless of payload type)
+				Swap(TempVolumeSet, Result.VolumeSet);
+				AGVListeners[i].Update(Result, ProxySearch.Location, AudioDeviceID);
+
+				if (FAudioGameplayVolumeProxyInfo* ProxyInfo = WorldProxyLists.Find(ProxySearch.WorldID))
+				{
+					ProxyInfo->AddListenerIndex(i);
+				}
 
-			if (FAudioGameplayVolumeProxyInfo* ProxyInfo = WorldProxyLists.Find(ProxySearch.WorldID))
-			{
-				ProxyInfo->AddListenerIndex(i);
+				continue;
 			}
 		}
+
+		// Listener is invalid or uninitialized
+		Result.Reset();
+		AGVListeners[i].Update(Result, FVector::ZeroVector, AudioDeviceID);
 	}
 
 	FAudioGameplayProxyUpdateResult ProxyUpdateResult;
