AGV - downgrade ensure to warning in cases where the owning object has multiple primitive components, and remove ensure on data change notification.

diff --git a/Engine/Plugins/AudioGameplayVolume/Source/AudioGameplayVolume/Private/AudioGameplayVolumeMutator.cpp b/Engine/Plugins/AudioGameplayVolume/Source/AudioGameplayVolume/Private/AudioGameplayVolumeMutator.cpp
--- a/Engine/Plugins/AudioGameplayVolume/Source/AudioGameplayVolume/Private/AudioGameplayVolumeMutator.cpp
+++ b/Engine/Plugins/AudioGameplayVolume/Source/AudioGameplayVolume/Private/AudioGameplayVolumeMutator.cpp
@@ -162,11 +162,11 @@ void UAudioGameplayVolumeMutator::NotifyDataChanged() const
 	if (IsActive())
 	{
 		TInlineComponentArray<UAudioGameplayVolumeProxyComponent*> VolumeComponents(GetOwner());
-		if (ensureMsgf(VolumeComponents.Num() == 1, TEXT("Expecting exactly one AudioGameplayVolumeProxyComponent on an actor!")))
+		for (UAudioGameplayVolumeProxyComponent* VolumeComponent : VolumeComponents)
 		{
-			if (VolumeComponents[0])
+			if (VolumeComponent != nullptr)
 			{
-				VolumeComponents[0]->OnComponentDataChanged();
+				VolumeComponent->OnComponentDataChanged();
 			}
 		}
 	}
diff --git a/Engine/Plugins/AudioGameplayVolume/Source/AudioGameplayVolume/Private/AudioGameplayVolumeProxy.cpp b/Engine/Plugins/AudioGameplayVolume/Source/AudioGameplayVolume/Private/AudioGameplayVolumeProxy.cpp
--- a/Engine/Plugins/AudioGameplayVolume/Source/AudioGameplayVolume/Private/AudioGameplayVolumeProxy.cpp
+++ b/Engine/Plugins/AudioGameplayVolume/Source/AudioGameplayVolume/Private/AudioGameplayVolumeProxy.cpp
@@ -147,7 +147,14 @@ void UAGVPrimitiveComponentProxy::InitFromComponent(const UAudioGameplayVolumePr
 	if (Component)
 	{
 		TInlineComponentArray<UPrimitiveComponent*> PrimitiveComponents(Component->GetOwner());
-		if (ensureMsgf(PrimitiveComponents.Num() == 1, TEXT("An Audio Gameplay Volume Shape Proxy requires exactly one Primitive Component on the owning actor")))
+		const int32 PrimitiveComponentCount = PrimitiveComponents.Num();
+
+		if (PrimitiveComponentCount != 1)
+		{
+			UE_LOG(AudioGameplayVolumeLog, Warning, TEXT("Was expecting exactly one Primitive Component on the owning actor, found %d - this could cause unexpected behavior"), PrimitiveComponentCount);
+		}
+
+		if (PrimitiveComponents.Num() > 0)
 		{
 			WeakPrimative = PrimitiveComponents[0];
 		}
