Fix for audio gameplay components not getting listener messages

diff --git a/Engine/Plugins/AudioGameplayVolume/Source/AudioGameplayVolume/Private/AudioGameplayVolumeComponent.cpp b/Engine/Plugins/AudioGameplayVolume/Source/AudioGameplayVolume/Private/AudioGameplayVolumeComponent.cpp
--- a/Engine/Plugins/AudioGameplayVolume/Source/AudioGameplayVolume/Private/AudioGameplayVolumeComponent.cpp
+++ b/Engine/Plugins/AudioGameplayVolume/Source/AudioGameplayVolume/Private/AudioGameplayVolumeComponent.cpp
@@ -36,12 +36,7 @@ void UAudioGameplayVolumeProxyComponent::EnterProxy() const
 	TInlineComponentArray<UActorComponent*> ActorComponents(GetOwner());
 	for (UActorComponent* ActorComponent : ActorComponents)
 	{
-		if (!ActorComponent || !ActorComponent->IsActive())
-		{
-			continue;
-		}
-
-		if (ActorComponent->Implements<UAudioGameplayVolumeInteraction>())
+		if (ActorComponent && ActorComponent->Implements<UAudioGameplayVolumeInteraction>())
 		{
 			IAudioGameplayVolumeInteraction::Execute_OnListenerEnter(ActorComponent);
 		}
@@ -55,12 +50,7 @@ void UAudioGameplayVolumeProxyComponent::ExitProxy() const
 	TInlineComponentArray<UActorComponent*> ActorComponents(GetOwner());
 	for (UActorComponent* ActorComponent : ActorComponents)
 	{
-		if (!ActorComponent || !ActorComponent->IsActive())
-		{
-			continue;
-		}
-
-		if (ActorComponent->Implements<UAudioGameplayVolumeInteraction>())
+		if (ActorComponent && ActorComponent->Implements<UAudioGameplayVolumeInteraction>())
 		{
 			IAudioGameplayVolumeInteraction::Execute_OnListenerExit(ActorComponent);
 		}
