Renaming AudioGameplayVolumeProxyComponent to AudioGameplayVolumeComponent - more compact name, better reflects intention of object

diff --git a/Engine/Plugins/AudioGameplayVolume/Config/BaseAudioGameplayVolume.ini b/Engine/Plugins/AudioGameplayVolume/Config/BaseAudioGameplayVolume.ini
new file mode 100644
--- /dev/null
+++ b/Engine/Plugins/AudioGameplayVolume/Config/BaseAudioGameplayVolume.ini
@@ -0,0 +1,3 @@
+
+[CoreRedirects]
++ClassRedirects=(OldName="/Script/AudioGameplayVolume.AudioGameplayVolumeProxyComponent", NewName="/Script/AudioGameplayVolume.AudioGameplayVolumeComponent")
diff --git a/Engine/Plugins/AudioGameplayVolume/Source/AudioGameplayVolume/Private/AudioGameplayVolume.cpp b/Engine/Plugins/AudioGameplayVolume/Source/AudioGameplayVolume/Private/AudioGameplayVolume.cpp
--- a/Engine/Plugins/AudioGameplayVolume/Source/AudioGameplayVolume/Private/AudioGameplayVolume.cpp
+++ b/Engine/Plugins/AudioGameplayVolume/Source/AudioGameplayVolume/Private/AudioGameplayVolume.cpp
@@ -19,7 +19,7 @@ AAudioGameplayVolume::AAudioGameplayVolume(const FObjectInitializer& ObjectIniti
 		BrushComp->bAlwaysCreatePhysicsState = true;
 	}
 
-	AGVComponent = CreateDefaultSubobject<UAudioGameplayVolumeProxyComponent>(TEXT("AGVComponent"));
+	AGVComponent = CreateDefaultSubobject<UAudioGameplayVolumeComponent>(TEXT("AGVComponent"));
 
 #if WITH_EDITOR
 	bColored = true;
diff --git a/Engine/Plugins/AudioGameplayVolume/Source/AudioGameplayVolume/Private/AudioGameplayVolumeComponent.cpp b/Engine/Plugins/AudioGameplayVolume/Source/AudioGameplayVolume/Private/AudioGameplayVolumeComponent.cpp
--- a/Engine/Plugins/AudioGameplayVolume/Source/AudioGameplayVolume/Private/AudioGameplayVolumeComponent.cpp
+++ b/Engine/Plugins/AudioGameplayVolume/Source/AudioGameplayVolume/Private/AudioGameplayVolumeComponent.cpp
@@ -6,13 +6,13 @@
 #include "Interfaces/IAudioGameplayVolumeInteraction.h"
 #include "AudioDevice.h"
 
-UAudioGameplayVolumeProxyComponent::UAudioGameplayVolumeProxyComponent(const FObjectInitializer& ObjectInitializer)
+UAudioGameplayVolumeComponent::UAudioGameplayVolumeComponent(const FObjectInitializer& ObjectInitializer)
 	: Super(ObjectInitializer)
 {
 	bAutoActivate = true;
 }
 
-void UAudioGameplayVolumeProxyComponent::SetProxy(UAudioGameplayVolumeProxy* NewProxy)
+void UAudioGameplayVolumeComponent::SetProxy(UAudioGameplayVolumeProxy* NewProxy)
 {
 	RemoveProxy();
 	Proxy = NewProxy;
@@ -23,7 +23,7 @@ void UAudioGameplayVolumeProxyComponent::SetProxy(UAudioGameplayVolumeProxy* New
 	}
 }
 
-void UAudioGameplayVolumeProxyComponent::OnComponentDataChanged()
+void UAudioGameplayVolumeComponent::OnComponentDataChanged()
 {
 	if (IsActive())
 	{
@@ -31,7 +31,7 @@ void UAudioGameplayVolumeProxyComponent::OnComponentDataChanged()
 	}
 }
 
-void UAudioGameplayVolumeProxyComponent::EnterProxy() const
+void UAudioGameplayVolumeComponent::EnterProxy() const
 {
 	TInlineComponentArray<UActorComponent*> ActorComponents(GetOwner());
 	for (UActorComponent* ActorComponent : ActorComponents)
@@ -45,7 +45,7 @@ void UAudioGameplayVolumeProxyComponent::EnterProxy() const
 	OnProxyEnter.Broadcast();
 }
 
-void UAudioGameplayVolumeProxyComponent::ExitProxy() const
+void UAudioGameplayVolumeComponent::ExitProxy() const
 {
 	TInlineComponentArray<UActorComponent*> ActorComponents(GetOwner());
 	for (UActorComponent* ActorComponent : ActorComponents)
@@ -60,11 +60,11 @@ void UAudioGameplayVolumeProxyComponent::ExitProxy() const
 }
 
 #if WITH_EDITOR
-void UAudioGameplayVolumeProxyComponent::PostEditChangeProperty(FPropertyChangedEvent& PropertyChangedEvent)
+void UAudioGameplayVolumeComponent::PostEditChangeProperty(FPropertyChangedEvent& PropertyChangedEvent)
 {
 	Super::PostEditChangeProperty(PropertyChangedEvent);
 
-	if (PropertyChangedEvent.Property && PropertyChangedEvent.Property->GetFName() == GET_MEMBER_NAME_CHECKED(UAudioGameplayVolumeProxyComponent, Proxy))
+	if (PropertyChangedEvent.Property && PropertyChangedEvent.Property->GetFName() == GET_MEMBER_NAME_CHECKED(UAudioGameplayVolumeComponent, Proxy))
 	{
 		RemoveProxy();
 
@@ -76,13 +76,13 @@ void UAudioGameplayVolumeProxyComponent::PostEditChangeProperty(FPropertyChanged
 }
 #endif // WITH_EDITOR
 
-void UAudioGameplayVolumeProxyComponent::OnUnregister()
+void UAudioGameplayVolumeComponent::OnUnregister()
 {
 	Super::OnUnregister();
 	RemoveProxy();
 }
 
-void UAudioGameplayVolumeProxyComponent::Enable()
+void UAudioGameplayVolumeComponent::Enable()
 {
 	if (Proxy != nullptr)
 	{
@@ -91,13 +91,13 @@ void UAudioGameplayVolumeProxyComponent::Enable()
 	}
 }
 
-void UAudioGameplayVolumeProxyComponent::Disable()
+void UAudioGameplayVolumeComponent::Disable()
 {
 	RemoveProxy();
 	Super::Disable();
 }
 
-void UAudioGameplayVolumeProxyComponent::AddProxy() const
+void UAudioGameplayVolumeComponent::AddProxy() const
 {
 	if (UAudioGameplayVolumeSubsystem* VolumeSubsystem = GetSubsystem())
 	{
@@ -105,7 +105,7 @@ void UAudioGameplayVolumeProxyComponent::AddProxy() const
 	}
 }
 
-void UAudioGameplayVolumeProxyComponent::RemoveProxy() const
+void UAudioGameplayVolumeComponent::RemoveProxy() const
 {
 	if (UAudioGameplayVolumeSubsystem* VolumeSubsystem = GetSubsystem())
 	{
@@ -113,7 +113,7 @@ void UAudioGameplayVolumeProxyComponent::RemoveProxy() const
 	}
 }
 
-void UAudioGameplayVolumeProxyComponent::UpdateProxy() const
+void UAudioGameplayVolumeComponent::UpdateProxy() const
 {
 	if (UAudioGameplayVolumeSubsystem* VolumeSubsystem = GetSubsystem())
 	{
@@ -121,7 +121,7 @@ void UAudioGameplayVolumeProxyComponent::UpdateProxy() const
 	}
 }
 
-UAudioGameplayVolumeSubsystem* UAudioGameplayVolumeProxyComponent::GetSubsystem() const
+UAudioGameplayVolumeSubsystem* UAudioGameplayVolumeComponent::GetSubsystem() const
 {
 	if (UWorld* World = GetWorld())
 	{
diff --git a/Engine/Plugins/AudioGameplayVolume/Source/AudioGameplayVolume/Private/AudioGameplayVolumeMutator.cpp b/Engine/Plugins/AudioGameplayVolume/Source/AudioGameplayVolume/Private/AudioGameplayVolumeMutator.cpp
--- a/Engine/Plugins/AudioGameplayVolume/Source/AudioGameplayVolume/Private/AudioGameplayVolumeMutator.cpp
+++ b/Engine/Plugins/AudioGameplayVolume/Source/AudioGameplayVolume/Private/AudioGameplayVolumeMutator.cpp
@@ -161,8 +161,8 @@ void UAudioGameplayVolumeMutator::NotifyDataChanged() const
 {
 	if (IsActive())
 	{
-		TInlineComponentArray<UAudioGameplayVolumeProxyComponent*> VolumeComponents(GetOwner());
-		for (UAudioGameplayVolumeProxyComponent* VolumeComponent : VolumeComponents)
+		TInlineComponentArray<UAudioGameplayVolumeComponent*> VolumeComponents(GetOwner());
+		for (UAudioGameplayVolumeComponent* VolumeComponent : VolumeComponents)
 		{
 			if (VolumeComponent != nullptr)
 			{
diff --git a/Engine/Plugins/AudioGameplayVolume/Source/AudioGameplayVolume/Private/AudioGameplayVolumeProxy.cpp b/Engine/Plugins/AudioGameplayVolume/Source/AudioGameplayVolume/Private/AudioGameplayVolumeProxy.cpp
--- a/Engine/Plugins/AudioGameplayVolume/Source/AudioGameplayVolume/Private/AudioGameplayVolumeProxy.cpp
+++ b/Engine/Plugins/AudioGameplayVolume/Source/AudioGameplayVolume/Private/AudioGameplayVolumeProxy.cpp
@@ -20,7 +20,7 @@ bool UAudioGameplayVolumeProxy::ContainsPosition(const FVector& Position) const
 	return false;
 }
 
-void UAudioGameplayVolumeProxy::InitFromComponent(const UAudioGameplayVolumeProxyComponent* Component)
+void UAudioGameplayVolumeProxy::InitFromComponent(const UAudioGameplayVolumeComponent* Component)
 {
 	if (!Component || !Component->GetWorld())
 	{
@@ -140,7 +140,7 @@ bool UAGVPrimitiveComponentProxy::ContainsPosition(const FVector& Position) cons
 	return BodyInstancePointer->GetSquaredDistanceToBody(Position, DistanceSquared, PointOnBody) && FMath::IsNearlyZero(DistanceSquared);
 }
 
-void UAGVPrimitiveComponentProxy::InitFromComponent(const UAudioGameplayVolumeProxyComponent* Component)
+void UAGVPrimitiveComponentProxy::InitFromComponent(const UAudioGameplayVolumeComponent* Component)
 {
 	Super::InitFromComponent(Component);
 
@@ -180,7 +180,7 @@ bool UAGVConditionProxy::ContainsPosition(const FVector& Position) const
 	return false;
 }
 
-void UAGVConditionProxy::InitFromComponent(const UAudioGameplayVolumeProxyComponent* Component)
+void UAGVConditionProxy::InitFromComponent(const UAudioGameplayVolumeComponent* Component)
 {
 	Super::InitFromComponent(Component);
 
diff --git a/Engine/Plugins/AudioGameplayVolume/Source/AudioGameplayVolume/Private/AudioGameplayVolumeSubsystem.cpp b/Engine/Plugins/AudioGameplayVolume/Source/AudioGameplayVolume/Private/AudioGameplayVolumeSubsystem.cpp
--- a/Engine/Plugins/AudioGameplayVolume/Source/AudioGameplayVolume/Private/AudioGameplayVolumeSubsystem.cpp
+++ b/Engine/Plugins/AudioGameplayVolume/Source/AudioGameplayVolume/Private/AudioGameplayVolumeSubsystem.cpp
@@ -274,7 +274,7 @@ void UAudioGameplayVolumeSubsystem::OnNotifyPendingDelete(const FActiveSound& Ac
 	}
 }
 
-void UAudioGameplayVolumeSubsystem::AddVolumeComponent(const UAudioGameplayVolumeProxyComponent* VolumeComponent)
+void UAudioGameplayVolumeSubsystem::AddVolumeComponent(const UAudioGameplayVolumeComponent* VolumeComponent)
 {
 	if (!VolumeComponent)
 	{
@@ -309,7 +309,7 @@ void UAudioGameplayVolumeSubsystem::AddVolumeComponent(const UAudioGameplayVolum
 	}
 }
 
-void UAudioGameplayVolumeSubsystem::UpdateVolumeComponent(const UAudioGameplayVolumeProxyComponent* VolumeComponent)
+void UAudioGameplayVolumeSubsystem::UpdateVolumeComponent(const UAudioGameplayVolumeComponent* VolumeComponent)
 {
 	if (!VolumeComponent || !VolumeComponent->GetWorld())
 	{
@@ -337,7 +337,7 @@ void UAudioGameplayVolumeSubsystem::UpdateVolumeComponent(const UAudioGameplayVo
 	}
 }
 
-void UAudioGameplayVolumeSubsystem::RemoveVolumeComponent(const UAudioGameplayVolumeProxyComponent* VolumeComponent)
+void UAudioGameplayVolumeSubsystem::RemoveVolumeComponent(const UAudioGameplayVolumeComponent* VolumeComponent)
 {
 	if (!VolumeComponent || !VolumeComponent->GetWorld())
 	{
@@ -450,7 +450,7 @@ void UAudioGameplayVolumeSubsystem::UpdateComponentsFromProxyInfo(const FAudioGa
 			continue;
 		}
 
-		if (const UAudioGameplayVolumeProxyComponent* ProxyComponent = AGVComponents[VolumeID])
+		if (const UAudioGameplayVolumeComponent* ProxyComponent = AGVComponents[VolumeID])
 		{
 			ProxyComponent->EnterProxy();
 		}
@@ -463,7 +463,7 @@ void UAudioGameplayVolumeSubsystem::UpdateComponentsFromProxyInfo(const FAudioGa
 			continue;
 		}
 
-		if (const UAudioGameplayVolumeProxyComponent* ProxyComponent = AGVComponents[VolumeID])
+		if (const UAudioGameplayVolumeComponent* ProxyComponent = AGVComponents[VolumeID])
 		{
 			ProxyComponent->ExitProxy();
 		}
diff --git a/Engine/Plugins/AudioGameplayVolume/Source/AudioGameplayVolume/Public/AudioGameplayVolume.h b/Engine/Plugins/AudioGameplayVolume/Source/AudioGameplayVolume/Public/AudioGameplayVolume.h
--- a/Engine/Plugins/AudioGameplayVolume/Source/AudioGameplayVolume/Public/AudioGameplayVolume.h
+++ b/Engine/Plugins/AudioGameplayVolume/Source/AudioGameplayVolume/Public/AudioGameplayVolume.h
@@ -9,7 +9,7 @@
 DECLARE_DYNAMIC_MULTICAST_DELEGATE(FAudioGameplayVolumeEvent);
 
 // Forward Declarations 
-class UAudioGameplayVolumeProxyComponent;
+class UAudioGameplayVolumeComponent;
 class UAudioGameplayVolumeSubsystem;
 
 /**
@@ -28,7 +28,7 @@ private:
 
 	// A representation of this volume for the audio thread
 	UPROPERTY()
-	UAudioGameplayVolumeProxyComponent* AGVComponent = nullptr;
+	UAudioGameplayVolumeComponent* AGVComponent = nullptr;
 
 	// Whether this volume is currently enabled.  Disabled volumes will not have a volume proxy, 
 	// and therefore will not be considered for intersection checks.
diff --git a/Engine/Plugins/AudioGameplayVolume/Source/AudioGameplayVolume/Public/AudioGameplayVolumeComponent.h b/Engine/Plugins/AudioGameplayVolume/Source/AudioGameplayVolume/Public/AudioGameplayVolumeComponent.h
--- a/Engine/Plugins/AudioGameplayVolume/Source/AudioGameplayVolume/Public/AudioGameplayVolumeComponent.h
+++ b/Engine/Plugins/AudioGameplayVolume/Source/AudioGameplayVolume/Public/AudioGameplayVolumeComponent.h
@@ -12,17 +12,17 @@ class UAudioGameplayVolumeSubsystem;
 DECLARE_DYNAMIC_MULTICAST_DELEGATE(FOnAudioGameplayVolumeProxyStateChange);
 
 /**
- *  UAudioGameplayVolumeProxyComponent - Component used to drive interaction with AudioGameplayVolumeSubsystem.
+ *  UAudioGameplayVolumeComponent - Component used to drive interaction with AudioGameplayVolumeSubsystem.
  *   NOTE: Do not inherit from this class, use UAudioGameplayVolumeMutator to create extendable functionality
  */
 UCLASS(Config = Game, ClassGroup = ("AudioGameplay"), meta = (BlueprintSpawnableComponent, IsBlueprintBase = false, DisplayName = "Volume Proxy"))
-class AUDIOGAMEPLAYVOLUME_API UAudioGameplayVolumeProxyComponent final : public UAudioGameplayComponent
+class AUDIOGAMEPLAYVOLUME_API UAudioGameplayVolumeComponent final : public UAudioGameplayComponent
 {
 	GENERATED_UCLASS_BODY()
 
 public:
 
-	virtual ~UAudioGameplayVolumeProxyComponent() = default;
+	virtual ~UAudioGameplayVolumeComponent() = default;
 
 	void SetProxy(UAudioGameplayVolumeProxy* NewProxy);
 	UAudioGameplayVolumeProxy* GetProxy() const { return Proxy; }
@@ -30,7 +30,7 @@ public:
 	/** Called by a component on same actor to notify our proxy may need updating */
 	void OnComponentDataChanged();
 
-	/** Called when the proxy is 'entered' - This is when the proxy goes from zero listeners to at lesat one. */
+	/** Called when the proxy is 'entered' - This is when the proxy goes from zero listeners to at least one. */
 	void EnterProxy() const;
 
 	/** Called when the proxy is 'exited' - This is when the proxy goes from at least one listeners to zero. */
diff --git a/Engine/Plugins/AudioGameplayVolume/Source/AudioGameplayVolume/Public/AudioGameplayVolumeProxy.h b/Engine/Plugins/AudioGameplayVolume/Source/AudioGameplayVolume/Public/AudioGameplayVolumeProxy.h
--- a/Engine/Plugins/AudioGameplayVolume/Source/AudioGameplayVolume/Public/AudioGameplayVolumeProxy.h
+++ b/Engine/Plugins/AudioGameplayVolume/Source/AudioGameplayVolume/Public/AudioGameplayVolumeProxy.h
@@ -13,7 +13,7 @@ struct FAudioProxyMutatorPriorities;
 struct FAudioProxyMutatorSearchResult;
 class FProxyVolumeMutator;
 class UPrimitiveComponent;
-class UAudioGameplayVolumeProxyComponent;
+class UAudioGameplayVolumeComponent;
 
 /**
  *  UAudioGameplayVolumeProxy - Abstract proxy used on audio thread to represent audio gameplay volumes.
@@ -31,7 +31,7 @@ public:
 	virtual ~UAudioGameplayVolumeProxy() = default;
 
 	virtual bool ContainsPosition(const FVector& Position) const;
-	virtual void InitFromComponent(const UAudioGameplayVolumeProxyComponent* Component);
+	virtual void InitFromComponent(const UAudioGameplayVolumeComponent* Component);
 
 	void FindMutatorPriority(FAudioProxyMutatorPriorities& Priorities) const;
 	void GatherMutators(const FAudioProxyMutatorPriorities& Priorities, FAudioProxyMutatorSearchResult& OutResult) const;
@@ -64,7 +64,7 @@ public:
 	virtual ~UAGVPrimitiveComponentProxy() = default;
 
 	virtual bool ContainsPosition(const FVector& Position) const override;
-	virtual void InitFromComponent(const UAudioGameplayVolumeProxyComponent* Component) override;
+	virtual void InitFromComponent(const UAudioGameplayVolumeComponent* Component) override;
 
 protected:
 	
@@ -84,7 +84,7 @@ public:
 	virtual ~UAGVConditionProxy() = default;
 
 	virtual bool ContainsPosition(const FVector& Position) const override;
-	virtual void InitFromComponent(const UAudioGameplayVolumeProxyComponent* Component) override;
+	virtual void InitFromComponent(const UAudioGameplayVolumeComponent* Component) override;
 
 protected:
 
diff --git a/Engine/Plugins/AudioGameplayVolume/Source/AudioGameplayVolume/Public/AudioGameplayVolumeSubsystem.h b/Engine/Plugins/AudioGameplayVolume/Source/AudioGameplayVolume/Public/AudioGameplayVolumeSubsystem.h
--- a/Engine/Plugins/AudioGameplayVolume/Source/AudioGameplayVolume/Public/AudioGameplayVolumeSubsystem.h
+++ b/Engine/Plugins/AudioGameplayVolume/Source/AudioGameplayVolume/Public/AudioGameplayVolumeSubsystem.h
@@ -13,7 +13,7 @@
 class UWorld;
 class FProxyVolumeMutator;
 class UAudioGameplayVolumeProxy;
-class UAudioGameplayVolumeProxyComponent;
+class UAudioGameplayVolumeComponent;
 
 /**
  *  FAudioGameplayActiveSoundInfo - Helper struct for caching active sound data
@@ -138,13 +138,13 @@ public:
 	//~ End IActiveSoundUpdateInterface
 	
 	/** Add a volume to the system */
-	void AddVolumeComponent(const UAudioGameplayVolumeProxyComponent* VolumeComponent);
+	void AddVolumeComponent(const UAudioGameplayVolumeComponent* VolumeComponent);
 
 	/** Update an existing volume in the system */
-	void UpdateVolumeComponent(const UAudioGameplayVolumeProxyComponent* VolumeComponent);
+	void UpdateVolumeComponent(const UAudioGameplayVolumeComponent* VolumeComponent);
 
 	/** Remove a volume from the system */
-	void RemoveVolumeComponent(const UAudioGameplayVolumeProxyComponent* VolumeComponent);
+	void RemoveVolumeComponent(const UAudioGameplayVolumeComponent* VolumeComponent);
 
 protected:
 
@@ -167,7 +167,7 @@ protected:
 
 	// Components in our system
 	UPROPERTY(Transient)
-	TMap<uint32, const UAudioGameplayVolumeProxyComponent*> AGVComponents;
+	TMap<uint32, const UAudioGameplayVolumeComponent*> AGVComponents;
 
 	// Audio thread representation of Listeners
 	TArray<FAudioGameplayVolumeListener> AGVListeners;
