Syntatic sugar for extending AudioGameplayVolume Components

diff --git a/Engine/Plugins/AudioGameplayVolume/Source/AudioGameplayVolume/Private/AudioGameplayVolumeComponent.cpp b/Engine/Plugins/AudioGameplayVolume/Source/AudioGameplayVolume/Private/AudioGameplayVolumeComponent.cpp
--- a/Engine/Plugins/AudioGameplayVolume/Source/AudioGameplayVolume/Private/AudioGameplayVolumeComponent.cpp
+++ b/Engine/Plugins/AudioGameplayVolume/Source/AudioGameplayVolume/Private/AudioGameplayVolumeComponent.cpp
@@ -129,4 +129,10 @@ UAudioGameplayVolumeSubsystem* UAudioGameplayVolumeComponent::GetSubsystem() con
 	}
 
 	return nullptr;
-}
\ No newline at end of file
+}
+
+UAudioGameplayVolumeComponentBase::UAudioGameplayVolumeComponentBase(const FObjectInitializer& ObjectInitializer)
+	: Super(ObjectInitializer)
+{
+	bAutoActivate = true;
+}
diff --git a/Engine/Plugins/AudioGameplayVolume/Source/AudioGameplayVolume/Public/AudioGameplayVolumeComponent.h b/Engine/Plugins/AudioGameplayVolume/Source/AudioGameplayVolume/Public/AudioGameplayVolumeComponent.h
--- a/Engine/Plugins/AudioGameplayVolume/Source/AudioGameplayVolume/Public/AudioGameplayVolumeComponent.h
+++ b/Engine/Plugins/AudioGameplayVolume/Source/AudioGameplayVolume/Public/AudioGameplayVolumeComponent.h
@@ -3,6 +3,7 @@
 #pragma once
 
 #include "AudioGameplayComponent.h"
+#include "Interfaces/IAudioGameplayVolumeInteraction.h"
 #include "AudioGameplayVolumeComponent.generated.h"
 
 // Forward Declarations 
@@ -13,7 +14,7 @@ DECLARE_DYNAMIC_MULTICAST_DELEGATE(FOnAudioGameplayVolumeProxyStateChange);
 
 /**
  *  UAudioGameplayVolumeComponent - Component used to drive interaction with AudioGameplayVolumeSubsystem.
- *   NOTE: Do not inherit from this class, use UAudioGameplayVolumeMutator to create extendable functionality
+ *   NOTE: Do not inherit from this class, use UAudioGameplayVolumeComponentBase or UAudioGameplayVolumeMutator to create extendable functionality
  */
 UCLASS(Config = Game, ClassGroup = ("AudioGameplay"), meta = (BlueprintSpawnableComponent, IsBlueprintBase = false, DisplayName = "Volume Proxy"))
 class AUDIOGAMEPLAYVOLUME_API UAudioGameplayVolumeComponent final : public UAudioGameplayComponent
@@ -68,4 +69,19 @@ protected:
 	void UpdateProxy() const;
 
 	UAudioGameplayVolumeSubsystem* GetSubsystem() const;
-};
\ No newline at end of file
+};
+
+/**
+ *  UAudioGameplayVolumeComponentBase - Blueprintable component used to craft custom functionality with AudioGameplayVolumes.
+ *   NOTE: Inherit from this class to get easy access to OnListenerEnter and OnListenerExit Blueprint Events
+ */
+UCLASS(Blueprintable, ClassGroup = ("AudioGameplay"), hidecategories = (Tags, Collision), meta = (BlueprintSpawnableComponent))
+class AUDIOGAMEPLAYVOLUME_API UAudioGameplayVolumeComponentBase : public UAudioGameplayComponent
+																, public IAudioGameplayVolumeInteraction
+{
+	GENERATED_UCLASS_BODY()
+
+public:
+
+	virtual ~UAudioGameplayVolumeComponentBase() = default;
+};
