Allowing AudioGameplayVolume components to be added to AudioGameplayVolumes at runtime.

diff --git a/Engine/Plugins/AudioGameplayVolume/Source/AudioGameplayVolume/Private/AudioGameplayVolume.cpp b/Engine/Plugins/AudioGameplayVolume/Source/AudioGameplayVolume/Private/AudioGameplayVolume.cpp
--- a/Engine/Plugins/AudioGameplayVolume/Source/AudioGameplayVolume/Private/AudioGameplayVolume.cpp
+++ b/Engine/Plugins/AudioGameplayVolume/Source/AudioGameplayVolume/Private/AudioGameplayVolume.cpp
@@ -99,8 +99,12 @@ void AAudioGameplayVolume::PostRegisterAllComponents()
 
 	if (AGVComponent)
 	{
-		UAGVPrimitiveComponentProxy* PrimitiveComponentProxy = NewObject<UAGVPrimitiveComponentProxy>(AGVComponent);
-		AGVComponent->SetProxy(PrimitiveComponentProxy);
+		UAGVPrimitiveComponentProxy* PrimitiveComponentProxy = Cast<UAGVPrimitiveComponentProxy>(AGVComponent->GetProxy());
+		if (!PrimitiveComponentProxy)
+		{
+			AGVComponent->SetProxy(NewObject<UAGVPrimitiveComponentProxy>(AGVComponent));
+		}
+
 		AGVComponent->bAutoActivate = false;
 	}
 }
diff --git a/Engine/Plugins/AudioGameplayVolume/Source/AudioGameplayVolume/Private/AudioGameplayVolumeMutator.cpp b/Engine/Plugins/AudioGameplayVolume/Source/AudioGameplayVolume/Private/AudioGameplayVolumeMutator.cpp
--- a/Engine/Plugins/AudioGameplayVolume/Source/AudioGameplayVolume/Private/AudioGameplayVolumeMutator.cpp
+++ b/Engine/Plugins/AudioGameplayVolume/Source/AudioGameplayVolume/Private/AudioGameplayVolumeMutator.cpp
@@ -152,6 +152,12 @@ TSharedPtr<FProxyVolumeMutator> UAudioGameplayVolumeMutator::CreateMutator() con
 	return ProxyMutator;
 }
 
+void UAudioGameplayVolumeMutator::Enable()
+{
+	Super::Enable();
+	NotifyDataChanged();
+}
+
 TSharedPtr<FProxyVolumeMutator> UAudioGameplayVolumeMutator::FactoryMutator() const
 {
 	return TSharedPtr<FProxyVolumeMutator>();
diff --git a/Engine/Plugins/AudioGameplayVolume/Source/AudioGameplayVolume/Public/AudioGameplayVolume.h b/Engine/Plugins/AudioGameplayVolume/Source/AudioGameplayVolume/Public/AudioGameplayVolume.h
--- a/Engine/Plugins/AudioGameplayVolume/Source/AudioGameplayVolume/Public/AudioGameplayVolume.h
+++ b/Engine/Plugins/AudioGameplayVolume/Source/AudioGameplayVolume/Public/AudioGameplayVolume.h
@@ -19,7 +19,7 @@ class UAudioGameplayVolumeSubsystem;
 
    NOTE: Will only impact audio sources that have "apply ambient volumes" set on their sound class.
  */
-UCLASS(Blueprintable, hidecategories = (Advanced, Attachment, Collision, Volume))
+UCLASS(hidecategories = (Advanced, Attachment, Collision, Volume))
 class AUDIOGAMEPLAYVOLUME_API AAudioGameplayVolume : public AVolume
 {
 	GENERATED_UCLASS_BODY()
@@ -27,7 +27,7 @@ class AUDIOGAMEPLAYVOLUME_API AAudioGameplayVolume : public AVolume
 private:
 
 	// A representation of this volume for the audio thread
-	UPROPERTY()
+	UPROPERTY(Transient)
 	TObjectPtr<UAudioGameplayVolumeComponent> AGVComponent = nullptr;
 
 	// Whether this volume is currently enabled.  Disabled volumes will not have a volume proxy, 
diff --git a/Engine/Plugins/AudioGameplayVolume/Source/AudioGameplayVolume/Public/AudioGameplayVolumeMutator.h b/Engine/Plugins/AudioGameplayVolume/Source/AudioGameplayVolume/Public/AudioGameplayVolumeMutator.h
--- a/Engine/Plugins/AudioGameplayVolume/Source/AudioGameplayVolume/Public/AudioGameplayVolumeMutator.h
+++ b/Engine/Plugins/AudioGameplayVolume/Source/AudioGameplayVolume/Public/AudioGameplayVolumeMutator.h
@@ -94,7 +94,7 @@ protected:
  *		Currently, components inheriting this base can affect interior settings as well as active sounds or the audio listener(s) inside the volume.
  *		See also: FilterVolumeComponent, AttenuationVolumeComponent, SubmixSendComponent, SubmixOverrideVolumeComponent, and ReverbVolumeComponent
  */
-UCLASS(Abstract, HideDropdown)
+UCLASS(Abstract, HideDropdown, meta = (IsBlueprintBase = false))
 class AUDIOGAMEPLAYVOLUME_API UAudioGameplayVolumeMutator : public UAudioGameplayComponent
 {
 	GENERATED_UCLASS_BODY()
@@ -113,6 +113,10 @@ public:
 
 protected:
 
+	//~ Begin UAudioGameplayComponent interface 
+	virtual void Enable() override;
+	//~ End UAudioGameplayComponent interface
+	
 	/** Create this component's type of mutator */
 	virtual TSharedPtr<FProxyVolumeMutator> FactoryMutator() const;
 
