Add detail customization to Reverb (Volume) component
 + "ApplyReverb" checkbox does nothing - we're reusing an engine structure under the hood, and our system does not use this property.

diff --git a/Engine/Plugins/AudioGameplayVolume/AudioGameplayVolume.uplugin b/Engine/Plugins/AudioGameplayVolume/AudioGameplayVolume.uplugin
--- a/Engine/Plugins/AudioGameplayVolume/AudioGameplayVolume.uplugin
+++ b/Engine/Plugins/AudioGameplayVolume/AudioGameplayVolume.uplugin
@@ -18,6 +18,11 @@
 			"Name": "AudioGameplayVolume",
 			"Type": "Runtime",
 			"LoadingPhase": "Default"
+		},
+		{
+			"Name": "AudioGameplayVolumeEditor",
+			"Type": "Editor",
+			"LoadingPhase": "PostDefault"
 		}
 	],
 	"Plugins" : [
diff --git a/Engine/Plugins/AudioGameplayVolume/Source/AudioGameplayVolume/Public/ReverbVolumeComponent.h b/Engine/Plugins/AudioGameplayVolume/Source/AudioGameplayVolume/Public/ReverbVolumeComponent.h
--- a/Engine/Plugins/AudioGameplayVolume/Source/AudioGameplayVolume/Public/ReverbVolumeComponent.h
+++ b/Engine/Plugins/AudioGameplayVolume/Source/AudioGameplayVolume/Public/ReverbVolumeComponent.h
@@ -43,6 +43,8 @@ public:
 
 	const FReverbSettings& GetReverbSettings() const { return ReverbSettings; }
 
+	friend class FReverbVolumeComponentDetail;
+
 private:
 
 	//~ Begin UAudioGameplayVolumeMutator interface
diff --git a/Engine/Plugins/AudioGameplayVolume/Source/AudioGameplayVolumeEditor/AudioGameplayVolumeEditor.Build.cs b/Engine/Plugins/AudioGameplayVolume/Source/AudioGameplayVolumeEditor/AudioGameplayVolumeEditor.Build.cs
new file mode 100644
--- /dev/null
+++ b/Engine/Plugins/AudioGameplayVolume/Source/AudioGameplayVolumeEditor/AudioGameplayVolumeEditor.Build.cs
@@ -0,0 +1,38 @@
+// Copyright Epic Games, Inc. All Rights Reserved.
+
+using UnrealBuildTool;
+
+public class AudioGameplayVolumeEditor : ModuleRules
+{
+	public AudioGameplayVolumeEditor(ReadOnlyTargetRules Target) : base(Target)
+	{
+		PublicDependencyModuleNames.AddRange(
+			new string[]
+			{
+				"Core",
+				"AudioGameplayVolume",
+				"ToolMenus",
+			}
+			);
+			
+		PrivateDependencyModuleNames.AddRange(
+			new string[]
+			{
+				"CoreUObject",
+				"Engine",
+				"Slate",
+				"UnrealEd",
+				"SlateCore",
+				"PropertyEditor",
+				"AudioGameplay",
+			}
+			);
+
+		PrivateIncludePathModuleNames.AddRange(
+			new string[]
+			{
+				"AssetTools",
+			}
+			);
+	}
+}
diff --git a/Engine/Plugins/AudioGameplayVolume/Source/AudioGameplayVolumeEditor/Private/AudioGameplayVolumeEditorModule.cpp b/Engine/Plugins/AudioGameplayVolume/Source/AudioGameplayVolumeEditor/Private/AudioGameplayVolumeEditorModule.cpp
new file mode 100644
--- /dev/null
+++ b/Engine/Plugins/AudioGameplayVolume/Source/AudioGameplayVolumeEditor/Private/AudioGameplayVolumeEditorModule.cpp
@@ -0,0 +1,36 @@
+// Copyright Epic Games, Inc. All Rights Reserved.
+
+#include "AudioGameplayVolumeEditorModule.h"
+#include "Modules/ModuleManager.h"
+#include "AssetToolsModule.h"
+#include "ToolMenus.h"
+#include "DetailCustomizations/ReverbVolumeComponentDetail.h"
+
+void FAudioGameplayVolumeEditorModule::StartupModule()
+{
+	constexpr bool bRegisterLayouts = true;
+	HandleCustomPropertyLayouts(bRegisterLayouts);
+}
+
+void FAudioGameplayVolumeEditorModule::ShutdownModule()
+{
+	constexpr bool bRegisterLayouts = false;
+	HandleCustomPropertyLayouts(bRegisterLayouts);
+}
+
+void FAudioGameplayVolumeEditorModule::HandleCustomPropertyLayouts(bool bRegisterLayouts)
+{
+	// Register detail customizations
+	FPropertyEditorModule& PropertyModule = FModuleManager::LoadModuleChecked<FPropertyEditorModule>(TEXT("PropertyEditor"));
+
+	if (bRegisterLayouts)
+	{
+		PropertyModule.RegisterCustomClassLayout("ReverbVolumeComponent", FOnGetDetailCustomizationInstance::CreateStatic(&FReverbVolumeComponentDetail::MakeInstance));
+	}
+	else
+	{
+		PropertyModule.UnregisterCustomClassLayout("ReverbVolumeComponent");
+	}
+}
+
+IMPLEMENT_MODULE(FAudioGameplayVolumeEditorModule, AudioGameplayVolumeEditor);
diff --git a/Engine/Plugins/AudioGameplayVolume/Source/AudioGameplayVolumeEditor/Private/DetailCustomizations/ReverbVolumeComponentDetail.cpp b/Engine/Plugins/AudioGameplayVolume/Source/AudioGameplayVolumeEditor/Private/DetailCustomizations/ReverbVolumeComponentDetail.cpp
new file mode 100644
--- /dev/null
+++ b/Engine/Plugins/AudioGameplayVolume/Source/AudioGameplayVolumeEditor/Private/DetailCustomizations/ReverbVolumeComponentDetail.cpp
@@ -0,0 +1,65 @@
+// Copyright Epic Games, Inc. All Rights Reserved.
+
+#include "ReverbVolumeComponentDetail.h"
+
+#include "DetailCategoryBuilder.h"
+#include "DetailLayoutBuilder.h"
+#include "PropertyCustomizationHelpers.h"
+#include "ReverbVolumeComponent.h"
+#include "ThumbnailRendering/ThumbnailManager.h"
+
+void FReverbVolumeComponentDetail::CustomizeDetails(IDetailLayoutBuilder& DetailBuilder)
+{
+	TArray<TWeakObjectPtr<UObject>> Objects;
+	DetailBuilder.GetObjectsBeingCustomized(Objects);
+
+	const int32 CustomizedObjects = Objects.Num();
+	ComponentToModify = nullptr;
+
+	for (TWeakObjectPtr<UObject>& Object : Objects)
+	{
+		if (Object->IsA<UReverbVolumeComponent>())
+		{
+			ComponentToModify = Cast<UReverbVolumeComponent>(Object.Get());
+			break;
+		}
+	}
+
+	if (!ComponentToModify)
+	{
+		return;
+	}
+
+	TSharedPtr<IPropertyHandle> ReverbSettingsProperty = DetailBuilder.GetProperty(GET_MEMBER_NAME_CHECKED(UReverbVolumeComponent, ReverbSettings));
+	check(ReverbSettingsProperty->IsValidHandle());
+
+	IDetailCategoryBuilder& ReverbCategory = DetailBuilder.EditCategory("Reverb");
+
+	uint32 NumChildren = 0;
+	ReverbSettingsProperty->GetNumChildren(NumChildren);
+	TArray<TSharedRef<IPropertyHandle>> ChildPropertyHandles;
+
+	for (uint32 ChildIndex = 0; ChildIndex < NumChildren; ChildIndex++)
+	{
+		ChildPropertyHandles.Add(ReverbSettingsProperty->GetChildHandle(ChildIndex).ToSharedRef());
+	}
+
+	for (const TSharedRef<IPropertyHandle>& ChildPropertyHandle : ChildPropertyHandles)
+	{
+		// Leave alone if already customized
+		if (ChildPropertyHandle->IsCustomized())
+		{
+			continue;
+		}
+
+		const FProperty* ChildProperty = ChildPropertyHandle->GetProperty();
+		const FName ChildPropertyName = ChildProperty->GetFName();
+
+		if (ChildPropertyName != GET_MEMBER_NAME_CHECKED(FReverbSettings, bApplyReverb))
+		{
+			ReverbCategory.AddProperty(ChildPropertyHandle);
+		}
+	}
+
+	DetailBuilder.HideProperty(ReverbSettingsProperty);
+}
diff --git a/Engine/Plugins/AudioGameplayVolume/Source/AudioGameplayVolumeEditor/Private/DetailCustomizations/ReverbVolumeComponentDetail.h b/Engine/Plugins/AudioGameplayVolume/Source/AudioGameplayVolumeEditor/Private/DetailCustomizations/ReverbVolumeComponentDetail.h
new file mode 100644
--- /dev/null
+++ b/Engine/Plugins/AudioGameplayVolume/Source/AudioGameplayVolumeEditor/Private/DetailCustomizations/ReverbVolumeComponentDetail.h
@@ -0,0 +1,25 @@
+// Copyright Epic Games, Inc. All Rights Reserved.
+
+#pragma once
+
+#include "IDetailCustomization.h"
+
+struct FAssetData;
+class UReverbVolumeComponent;
+class IDetailCategoryBuilder;
+
+class FReverbVolumeComponentDetail : public IDetailCustomization
+{
+public:
+
+	static TSharedRef<IDetailCustomization> MakeInstance() { return MakeShared<FReverbVolumeComponentDetail>(); }
+
+	virtual void CustomizeDetails(IDetailLayoutBuilder& DetailBuilder) override;
+
+private:
+
+	static void SortCategories(const TMap<FName, IDetailCategoryBuilder*>& AllCategories);
+	bool ShouldFilterAsset(const FAssetData& AssetData);
+
+	TObjectPtr<UReverbVolumeComponent> ComponentToModify = nullptr;
+};
diff --git a/Engine/Plugins/AudioGameplayVolume/Source/AudioGameplayVolumeEditor/Public/AudioGameplayVolumeEditorModule.h b/Engine/Plugins/AudioGameplayVolume/Source/AudioGameplayVolumeEditor/Public/AudioGameplayVolumeEditorModule.h
new file mode 100644
--- /dev/null
+++ b/Engine/Plugins/AudioGameplayVolume/Source/AudioGameplayVolumeEditor/Public/AudioGameplayVolumeEditorModule.h
@@ -0,0 +1,20 @@
+// Copyright Epic Games, Inc. All Rights Reserved.
+
+#pragma once
+
+#include "Modules/ModuleInterface.h"
+#include "Logging/LogMacros.h"
+
+AUDIOGAMEPLAYVOLUMEEDITOR_API DECLARE_LOG_CATEGORY_EXTERN(AudioGameplayVolumeEditor, Log, All);
+
+class FAudioGameplayVolumeEditorModule : public IModuleInterface
+{
+public:
+	/** IModuleInterface implementation */
+	virtual void StartupModule() override;
+	virtual void ShutdownModule() override;
+
+protected:
+
+	void HandleCustomPropertyLayouts(bool bRegisterLayouts);
+};
