[CIS] Fix non-unity errors

diff --git a/Engine/Plugins/AudioGameplayVolume/Source/AudioGameplayVolumeEditor/Private/AudioGameplayVolumeEditorModule.cpp b/Engine/Plugins/AudioGameplayVolume/Source/AudioGameplayVolumeEditor/Private/AudioGameplayVolumeEditorModule.cpp
--- a/Engine/Plugins/AudioGameplayVolume/Source/AudioGameplayVolumeEditor/Private/AudioGameplayVolumeEditorModule.cpp
+++ b/Engine/Plugins/AudioGameplayVolume/Source/AudioGameplayVolumeEditor/Private/AudioGameplayVolumeEditorModule.cpp
@@ -3,6 +3,7 @@
 #include "AudioGameplayVolumeEditorModule.h"
 #include "Modules/ModuleManager.h"
 #include "AssetToolsModule.h"
+#include "PropertyEditorModule.h"
 #include "ToolMenus.h"
 #include "DetailCustomizations/ReverbVolumeComponentDetail.h"
 
diff --git a/Engine/Plugins/AudioGameplayVolume/Source/AudioGameplayVolumeEditor/Private/DetailCustomizations/ReverbVolumeComponentDetail.h b/Engine/Plugins/AudioGameplayVolume/Source/AudioGameplayVolumeEditor/Private/DetailCustomizations/ReverbVolumeComponentDetail.h
--- a/Engine/Plugins/AudioGameplayVolume/Source/AudioGameplayVolumeEditor/Private/DetailCustomizations/ReverbVolumeComponentDetail.h
+++ b/Engine/Plugins/AudioGameplayVolume/Source/AudioGameplayVolumeEditor/Private/DetailCustomizations/ReverbVolumeComponentDetail.h
@@ -3,6 +3,7 @@
 #pragma once
 
 #include "IDetailCustomization.h"
+#include "UObject/ObjectPtr.h"
 
 struct FAssetData;
 class UReverbVolumeComponent;
