Removing blueprint write access on Attenuation and Filter volume components

diff --git a/Engine/Plugins/AudioGameplayVolume/Source/AudioGameplayVolume/Public/AttenuationVolumeComponent.h b/Engine/Plugins/AudioGameplayVolume/Source/AudioGameplayVolume/Public/AttenuationVolumeComponent.h
--- a/Engine/Plugins/AudioGameplayVolume/Source/AudioGameplayVolume/Public/AttenuationVolumeComponent.h
+++ b/Engine/Plugins/AudioGameplayVolume/Source/AudioGameplayVolume/Public/AttenuationVolumeComponent.h
@@ -60,18 +60,18 @@ private:
 	//~ End UAudioGameplayVolumeMutator interface
 
 	// The desired volume of sounds outside the volume when the player is inside the volume
-	UPROPERTY(EditAnywhere, BlueprintReadWrite, Category = "VolumeAttenuation", meta = (AllowPrivateAccess = "true"))
+	UPROPERTY(EditAnywhere, BlueprintReadOnly, Category = "VolumeAttenuation", meta = (AllowPrivateAccess = "true"))
 	float ExteriorVolume = 1.0f;
 
 	// The time over which to interpolate from the current volume to the desired volume of sounds outside the volume when the player enters the volume
-	UPROPERTY(EditAnywhere, BlueprintReadWrite, Category = "VolumeAttenuation", meta = (AllowPrivateAccess = "true"))
+	UPROPERTY(EditAnywhere, BlueprintReadOnly, Category = "VolumeAttenuation", meta = (AllowPrivateAccess = "true"))
 	float ExteriorTime = 0.5f;
 
 	// The desired volume of sounds inside the volume when the player is outside the volume
-	UPROPERTY(EditAnywhere, BlueprintReadWrite, Category = "VolumeAttenuation", meta = (AllowPrivateAccess = "true"))
+	UPROPERTY(EditAnywhere, BlueprintReadOnly, Category = "VolumeAttenuation", meta = (AllowPrivateAccess = "true"))
 	float InteriorVolume = 1.0f;
 
 	// The time over which to interpolate from the current volume to the desired volume of sounds inside the volume when the player enters the volume
-	UPROPERTY(EditAnywhere, BlueprintReadWrite, Category = "VolumeAttenuation", meta = (AllowPrivateAccess = "true"))
+	UPROPERTY(EditAnywhere, BlueprintReadOnly, Category = "VolumeAttenuation", meta = (AllowPrivateAccess = "true"))
 	float InteriorTime = 0.5f;
 };
\ No newline at end of file
diff --git a/Engine/Plugins/AudioGameplayVolume/Source/AudioGameplayVolume/Public/FilterVolumeComponent.h b/Engine/Plugins/AudioGameplayVolume/Source/AudioGameplayVolume/Public/FilterVolumeComponent.h
--- a/Engine/Plugins/AudioGameplayVolume/Source/AudioGameplayVolume/Public/FilterVolumeComponent.h
+++ b/Engine/Plugins/AudioGameplayVolume/Source/AudioGameplayVolume/Public/FilterVolumeComponent.h
@@ -60,18 +60,18 @@ private:
 	//~ End UAudioGameplayVolumeMutator interface
 
 	// The desired LPF frequency cutoff (in hertz) of sounds outside the volume when the player is inside the volume
-	UPROPERTY(EditAnywhere, BlueprintReadWrite, Category = "VolumeFilter", meta = (AllowPrivateAccess = "true"))
+	UPROPERTY(EditAnywhere, BlueprintReadOnly, Category = "VolumeFilter", meta = (AllowPrivateAccess = "true"))
 	float ExteriorLPF = MAX_FILTER_FREQUENCY;
 
 	// The time over which to interpolate from the current LPF to the desired LPF of sounds outside the volume when the player enters the volume
-	UPROPERTY(EditAnywhere, BlueprintReadWrite, Category = "VolumeFilter", meta = (AllowPrivateAccess = "true"))
+	UPROPERTY(EditAnywhere, BlueprintReadOnly, Category = "VolumeFilter", meta = (AllowPrivateAccess = "true"))
 	float ExteriorLPFTime = 0.5f;
 
 	// The desired LPF frequency cutoff (in hertz) of sounds inside the volume when the player is outside the volume
-	UPROPERTY(EditAnywhere, BlueprintReadWrite, Category = "VolumeFilter", meta = (AllowPrivateAccess = "true"))
+	UPROPERTY(EditAnywhere, BlueprintReadOnly, Category = "VolumeFilter", meta = (AllowPrivateAccess = "true"))
 	float InteriorLPF = MAX_FILTER_FREQUENCY;
 
 	// The time over which to interpolate from the current LPF to the desired LPF of sounds inside the volume when the player enters the volume
-	UPROPERTY(EditAnywhere, BlueprintReadWrite, Category = "VolumeFilter", meta = (AllowPrivateAccess = "true"))
+	UPROPERTY(EditAnywhere, BlueprintReadOnly, Category = "VolumeFilter", meta = (AllowPrivateAccess = "true"))
 	float InteriorLPFTime = 0.5f;
 };
\ No newline at end of file
