Optimize AudioGameplayVolumeSubsystem UpdateFromListeners
+ added early distance culling to primitive proxies, reducing the frequency of physics body queries
+ CVar to disable culling (if something goes wrong): au.AudioGameplayVolumes.PrimitiveProxy.DistanceCulling=0

#rnx

diff --git a/Engine/Plugins/AudioGameplayVolume/Source/AudioGameplayVolume/Private/AudioGameplayVolumeProxy.cpp b/Engine/Plugins/AudioGameplayVolume/Source/AudioGameplayVolume/Private/AudioGameplayVolumeProxy.cpp
--- a/Engine/Plugins/AudioGameplayVolume/Source/AudioGameplayVolume/Private/AudioGameplayVolumeProxy.cpp
+++ b/Engine/Plugins/AudioGameplayVolume/Source/AudioGameplayVolume/Private/AudioGameplayVolumeProxy.cpp
@@ -10,6 +10,17 @@
 #include "Components/BrushComponent.h"
 #include "Components/PrimitiveComponent.h"
 
+namespace AudioGameplayVolumeConsoleVariables
+{
+	int32 bProxyDistanceCulling = 1;
+	FAutoConsoleVariableRef CVarProxyDistanceCulling(
+		TEXT("au.AudioGameplayVolumes.PrimitiveProxy.DistanceCulling"),
+		bProxyDistanceCulling,
+		TEXT("Skips physics body queries for proxies that are not close to the listener.\n0: Disable, 1: Enable (default)"),
+		ECVF_Default);
+
+} // namespace AudioGameplayVolumeConsoleVariables
+
 #include UE_INLINE_GENERATED_CPP_BY_NAME(AudioGameplayVolumeProxy)
 
 UAudioGameplayVolumeProxy::UAudioGameplayVolumeProxy(const FObjectInitializer& ObjectInitializer)
@@ -122,12 +133,12 @@ bool UAGVPrimitiveComponentProxy::ContainsPosition(const FVector& Position) cons
 	FBodyInstance* BodyInstancePointer = nullptr;
 	if (UPrimitiveComponent* PrimitiveComponent = WeakPrimative.Get())
 	{
-		if (PrimitiveComponent->IsPhysicsStateCreated() && PrimitiveComponent->HasValidPhysicsState())
+		if (NeedsPhysicsQuery(PrimitiveComponent, Position))
 		{
 			BodyInstancePointer = PrimitiveComponent->GetBodyInstance();
 		}
 	}
-	
+
 	if (!BodyInstancePointer)
 	{
 		return false;
@@ -159,6 +170,28 @@ void UAGVPrimitiveComponentProxy::InitFromComponent(const UAudioGameplayVolumeCo
 	}
 }
 
+bool UAGVPrimitiveComponentProxy::NeedsPhysicsQuery(UPrimitiveComponent* PrimitiveComponent, const FVector& Position) const
+{
+	check(PrimitiveComponent);
+
+	if (!PrimitiveComponent->IsPhysicsStateCreated() || !PrimitiveComponent->HasValidPhysicsState())
+	{
+		return false;
+	}
+
+	// Temporary kill switch for distance culling
+	if (AudioGameplayVolumeConsoleVariables::bProxyDistanceCulling == 0)
+	{
+		return true;
+	}
+
+	// Early distance culling
+	const float BoundsRadiusSq = FMath::Square(PrimitiveComponent->Bounds.SphereRadius);
+	const float DistanceSq = FVector::DistSquared(PrimitiveComponent->Bounds.Origin, Position);
+
+	return DistanceSq <= BoundsRadiusSq;
+}
+
 UAGVConditionProxy::UAGVConditionProxy(const FObjectInitializer& ObjectInitializer)
 	: Super(ObjectInitializer)
 {
diff --git a/Engine/Plugins/AudioGameplayVolume/Source/AudioGameplayVolume/Public/AudioGameplayVolumeProxy.h b/Engine/Plugins/AudioGameplayVolume/Source/AudioGameplayVolume/Public/AudioGameplayVolumeProxy.h
--- a/Engine/Plugins/AudioGameplayVolume/Source/AudioGameplayVolume/Public/AudioGameplayVolumeProxy.h
+++ b/Engine/Plugins/AudioGameplayVolume/Source/AudioGameplayVolume/Public/AudioGameplayVolumeProxy.h
@@ -68,6 +68,9 @@ public:
 
 protected:
 	
+	/** Returns false if we can skip the physics body query, which can be expensive */
+	bool NeedsPhysicsQuery(UPrimitiveComponent* PrimitiveComponent, const FVector& Position) const;
+
 	TWeakObjectPtr<UPrimitiveComponent> WeakPrimative;
 };
 
