AudioGameplayVolume improvements
+ Mutator components - allow component changes to update in real time while in PIE

#rnx

diff --git a/Engine/Plugins/AudioGameplayVolume/Source/AudioGameplayVolume/Private/AudioGameplayVolumeMutator.cpp b/Engine/Plugins/AudioGameplayVolume/Source/AudioGameplayVolume/Private/AudioGameplayVolumeMutator.cpp
--- a/Engine/Plugins/AudioGameplayVolume/Source/AudioGameplayVolume/Private/AudioGameplayVolumeMutator.cpp
+++ b/Engine/Plugins/AudioGameplayVolume/Source/AudioGameplayVolume/Private/AudioGameplayVolumeMutator.cpp
@@ -155,6 +155,18 @@ TSharedPtr<FProxyVolumeMutator> UAudioGameplayVolumeMutator::CreateMutator() con
 	return ProxyMutator;
 }
 
+#if WITH_EDITOR
+void UAudioGameplayVolumeMutator::PostEditChangeProperty(FPropertyChangedEvent& PropertyChangedEvent)
+{
+	Super::PostEditChangeProperty(PropertyChangedEvent);
+
+	if (PropertyChangedEvent.Property)
+	{
+		NotifyDataChanged();
+	}
+}
+#endif // WITH_EDITOR
+
 void UAudioGameplayVolumeMutator::Enable()
 {
 	Super::Enable();
diff --git a/Engine/Plugins/AudioGameplayVolume/Source/AudioGameplayVolume/Public/AudioGameplayVolumeMutator.h b/Engine/Plugins/AudioGameplayVolume/Source/AudioGameplayVolume/Public/AudioGameplayVolumeMutator.h
--- a/Engine/Plugins/AudioGameplayVolume/Source/AudioGameplayVolume/Public/AudioGameplayVolumeMutator.h
+++ b/Engine/Plugins/AudioGameplayVolume/Source/AudioGameplayVolume/Public/AudioGameplayVolumeMutator.h
@@ -111,6 +111,12 @@ public:
 	/** Create and fill the appropriate proxy mutator for this component */
 	virtual TSharedPtr<FProxyVolumeMutator> CreateMutator() const final;
 
+	//~ Begin UObject Interface
+#if WITH_EDITOR
+	virtual void PostEditChangeProperty(FPropertyChangedEvent& PropertyChangedEvent) override;
+#endif // WITH_EDITOR
+	//~ End UObject Interface
+
 protected:
 
 	//~ Begin UAudioGameplayComponent interface 
