AudioGameplayVolume fix for stationary sounds
+ Active Sounds currently playing need to have interior spaces invalidated when an AGV proxy is changed, otherwise they will not be affected by said changes.

#rnx

diff --git a/Engine/Plugins/AudioGameplayVolume/Source/AudioGameplayVolume/Private/AudioGameplayVolumeSubsystem.cpp b/Engine/Plugins/AudioGameplayVolume/Source/AudioGameplayVolume/Private/AudioGameplayVolumeSubsystem.cpp
--- a/Engine/Plugins/AudioGameplayVolume/Source/AudioGameplayVolume/Private/AudioGameplayVolumeSubsystem.cpp
+++ b/Engine/Plugins/AudioGameplayVolume/Source/AudioGameplayVolume/Private/AudioGameplayVolumeSubsystem.cpp
@@ -182,6 +182,16 @@ void UAudioGameplayVolumeSubsystem::Update()
 	NextUpdateDeltaTime = FMath::Max(AudioGameplayVolumeConsoleVariables::UpdateRate + JitterDelta, AudioGameplayVolumeConsoleVariables::MinUpdateRate);
 	TimeSinceUpdate = 0.f;
 
+	if (bHasStaleProxy)
+	{
+		FAudioDeviceHandle AudioDeviceHandle = GetAudioDeviceHandle();
+		if (AudioDeviceHandle.IsValid())
+		{
+			AudioDeviceHandle->InvalidateCachedInteriorVolumes();
+			bHasStaleProxy = false;
+		}
+	}
+
 	if (AudioGameplayVolumeConsoleVariables::bUpdateListeners != 0)
 	{
 		UpdateFromListeners();
@@ -387,11 +397,8 @@ bool UAudioGameplayVolumeSubsystem::AddProxy(TWeakObjectPtr<UAudioGameplayVolume
 	{
 		ProxyVolumes.Emplace(WeakProxy);
 
-		FAudioGameplayVolumeProxyInfo& ProxyInfo = WorldProxyLists.FindOrAdd(WeakProxy->GetWorldID());
-		if (ProxyInfo.IsVolumeInCurrentList(NewVolumeID))
-		{
-			bHasStaleProxy = true;
-		}
+		WorldProxyLists.FindOrAdd(WeakProxy->GetWorldID());
+		bHasStaleProxy = true;
 
 		UE_LOG(AudioGameplayVolumeLog, VeryVerbose, TEXT("Proxy [%08x] added"), NewVolumeID);
 		return true;
@@ -430,8 +437,14 @@ bool UAudioGameplayVolumeSubsystem::RemoveProxy(uint32 AudioGameplayVolumeID)
 		return !Proxy.IsValid() || Proxy->GetVolumeID() == AudioGameplayVolumeID;
 	});
 
-	UE_LOG(AudioGameplayVolumeLog, VeryVerbose, TEXT("Removed %d Proxies with Id [%08x]"), NumRemoved, AudioGameplayVolumeID);
-	return NumRemoved > 0;
+	if (NumRemoved > 0)
+	{
+		bHasStaleProxy = true;
+		UE_LOG(AudioGameplayVolumeLog, VeryVerbose, TEXT("Removed %d Proxies with Id [%08x]"), NumRemoved, AudioGameplayVolumeID);
+		return true;
+	}
+
+	return false;
 }
 
 bool UAudioGameplayVolumeSubsystem::IsAnyListenerInVolume(uint32 WorldID, uint32 VolumeID) const
