Non-unity compile fix

diff --git a/Engine/Plugins/AudioGameplayVolume/Source/AudioGameplayVolume/Public/SubmixSendVolumeComponent.h b/Engine/Plugins/AudioGameplayVolume/Source/AudioGameplayVolume/Public/SubmixSendVolumeComponent.h
--- a/Engine/Plugins/AudioGameplayVolume/Source/AudioGameplayVolume/Public/SubmixSendVolumeComponent.h
+++ b/Engine/Plugins/AudioGameplayVolume/Source/AudioGameplayVolume/Public/SubmixSendVolumeComponent.h
@@ -3,10 +3,9 @@
 #pragma once
 
 #include "AudioGameplayVolumeMutator.h"
+#include "Sound/AudioVolume.h"
 #include "SubmixSendVolumeComponent.generated.h"
 
-struct FAudioVolumeSubmixSendSettings;
-
 /**
  *  FProxyMutator_SubmixSend - An audio thread representation of Submix Sends
  */
