Dropping UAudioGameplayVolumeComponents on subsystem deinitialize

#rnx

#pf 64907a085406a7c330021ebc

diff --git a/Engine/Plugins/AudioGameplayVolume/Source/AudioGameplayVolume/Private/AudioGameplayVolumeSubsystem.cpp b/Engine/Plugins/AudioGameplayVolume/Source/AudioGameplayVolume/Private/AudioGameplayVolumeSubsystem.cpp
--- a/Engine/Plugins/AudioGameplayVolume/Source/AudioGameplayVolume/Private/AudioGameplayVolumeSubsystem.cpp
+++ b/Engine/Plugins/AudioGameplayVolume/Source/AudioGameplayVolume/Private/AudioGameplayVolumeSubsystem.cpp
@@ -160,6 +160,11 @@ bool UAudioGameplayVolumeSubsystem::ShouldCreateSubsystem(UObject* Outer) const
 	return !IsRunningDedicatedServer();
 }
 
+void UAudioGameplayVolumeSubsystem::Deinitialize()
+{
+	AGVComponents.Reset();
+}
+
 void UAudioGameplayVolumeSubsystem::Update()
 {
 	check(IsInAudioThread());
diff --git a/Engine/Plugins/AudioGameplayVolume/Source/AudioGameplayVolume/Public/AudioGameplayVolumeSubsystem.h b/Engine/Plugins/AudioGameplayVolume/Source/AudioGameplayVolume/Public/AudioGameplayVolumeSubsystem.h
--- a/Engine/Plugins/AudioGameplayVolume/Source/AudioGameplayVolume/Public/AudioGameplayVolumeSubsystem.h
+++ b/Engine/Plugins/AudioGameplayVolume/Source/AudioGameplayVolume/Public/AudioGameplayVolumeSubsystem.h
@@ -126,6 +126,7 @@ public:
 
 	//~ Begin USubsystem interface
 	virtual bool ShouldCreateSubsystem(UObject* Outer) const override;
+	virtual void Deinitialize() override;
 	//~ End USubsystem interface
 
 	//~ Begin UAudioEngineSubsystem interface
