[AGV] Fixing (front facing) incorrect parameter for filter component

diff --git a/Engine/Plugins/AudioGameplayVolume/Source/AudioGameplayVolume/Public/FilterVolumeComponent.h b/Engine/Plugins/AudioGameplayVolume/Source/AudioGameplayVolume/Public/FilterVolumeComponent.h
--- a/Engine/Plugins/AudioGameplayVolume/Source/AudioGameplayVolume/Public/FilterVolumeComponent.h
+++ b/Engine/Plugins/AudioGameplayVolume/Source/AudioGameplayVolume/Public/FilterVolumeComponent.h
@@ -41,13 +41,13 @@ public:
 	virtual ~UFilterVolumeComponent() = default;
 
 	UFUNCTION(BlueprintCallable, Category = "AudioGameplay")
-	void SetExteriorLPF(float Volume, float InterpolateTime);
+	void SetExteriorLPF(UPARAM(DisplayName = "Frequency") float Volume, float InterpolateTime);
 
 	float GetExteriorLPF() const { return ExteriorLPF; }
 	float GetExteriorLPFTime() const { return ExteriorLPFTime; }
 
 	UFUNCTION(BlueprintCallable, Category = "AudioGameplay")
-	void SetInteriorLPF(float Volume, float InterpolateTime);
+	void SetInteriorLPF(UPARAM(DisplayName = "Frequency") float Volume, float InterpolateTime);
 
 	float GetInteriorLPF() const { return InteriorLPF; }
 	float GetInteriorLPFTime() const { return InteriorLPFTime; }
