Crash fix for Stale subsystem on AudioThread

diff --git a/Engine/Plugins/AudioGameplayVolume/Source/AudioGameplayVolume/Private/AudioGameplayVolumeSubsystem.cpp b/Engine/Plugins/AudioGameplayVolume/Source/AudioGameplayVolume/Private/AudioGameplayVolumeSubsystem.cpp
--- a/Engine/Plugins/AudioGameplayVolume/Source/AudioGameplayVolume/Private/AudioGameplayVolumeSubsystem.cpp
+++ b/Engine/Plugins/AudioGameplayVolume/Source/AudioGameplayVolume/Private/AudioGameplayVolumeSubsystem.cpp
@@ -43,6 +43,13 @@ namespace AudioGameplayVolumeConsoleVariables
 		TEXT("A random delta to add to update rate to avoid performance heartbeats."),
 		ECVF_Default);
 
+	int32 bAudioThreadCmdRollback = 0;
+	FAutoConsoleVariableRef CVarAudioThreadCmdRollback(
+		TEXT("au.AudioGameplayVolumes.AudioThreadCmdRollback"),
+		bAudioThreadCmdRollback,
+		TEXT("When non-zero, uses old code to rollback late thread command change."),
+		ECVF_Default);
+
 } // namespace AudioGameplayVolumeConsoleVariables
 
 void FAudioGameplayActiveSoundInfo::Update(double ListenerInteriorStartTime)
@@ -372,14 +379,35 @@ void UAudioGameplayVolumeSubsystem::RemoveVolumeComponent(const UAudioGameplayVo
 	}
 
 	// Remove representation of volume from audio thread
-	TWeakObjectPtr<UAudioGameplayVolumeSubsystem> WeakThis(this);
-	FAudioThread::RunCommandOnAudioThread([WeakThis, ComponentID]()
+	if (AudioGameplayVolumeConsoleVariables::bAudioThreadCmdRollback)
 	{
-		if (WeakThis.IsValid())
+		TWeakObjectPtr<UAudioGameplayVolumeSubsystem> WeakThis(this);
+		FAudioThread::RunCommandOnAudioThread([WeakThis, ComponentID]()
 		{
-			WeakThis->RemoveProxy(ComponentID);
-		}
-	});
+			if (WeakThis.IsValid())
+			{
+				WeakThis->RemoveProxy(ComponentID);
+			}
+		});
+	}
+	else
+	{
+		Audio::FDeviceId CurrentDeviceId = GetAudioDeviceId();
+		FAudioThread::RunCommandOnAudioThread([CurrentDeviceId, ComponentID]()
+		{
+			if (FAudioDeviceManager* AudioDeviceManager = FAudioDeviceManager::Get())
+			{
+				FAudioDeviceHandle DeviceHandle = AudioDeviceManager->GetAudioDevice(CurrentDeviceId);
+				if (DeviceHandle.IsValid())
+				{
+					if (UAudioGameplayVolumeSubsystem* AGVSubsystem = DeviceHandle->GetSubsystem<UAudioGameplayVolumeSubsystem>())
+					{
+						AGVSubsystem->RemoveProxy(ComponentID);
+					}
+				}
+			}
+		});
+	}
 }
 
 bool UAudioGameplayVolumeSubsystem::DoesSupportWorld(UWorld* World) const
diff --git a/Engine/Source/Runtime/Engine/Private/Subsystems/AudioEngineSubsystem.cpp b/Engine/Source/Runtime/Engine/Private/Subsystems/AudioEngineSubsystem.cpp
--- a/Engine/Source/Runtime/Engine/Private/Subsystems/AudioEngineSubsystem.cpp
+++ b/Engine/Source/Runtime/Engine/Private/Subsystems/AudioEngineSubsystem.cpp
@@ -12,14 +12,21 @@ UAudioEngineSubsystem::UAudioEngineSubsystem()
 {
 }
 
-FAudioDeviceHandle UAudioEngineSubsystem::GetAudioDeviceHandle() const
+Audio::FDeviceId UAudioEngineSubsystem::GetAudioDeviceId() const
 {
 	const UAudioSubsystemCollectionRoot* SubsystemRoot = Cast<UAudioSubsystemCollectionRoot>(GetOuter());
 	check(SubsystemRoot);
 
+	return SubsystemRoot->GetAudioDeviceID();
+}
+
+FAudioDeviceHandle UAudioEngineSubsystem::GetAudioDeviceHandle() const
+{
+	Audio::FDeviceId DeviceId = GetAudioDeviceId();
+
 	if (FAudioDeviceManager* AudioDeviceManager = FAudioDeviceManager::Get())
 	{
-		return AudioDeviceManager->GetAudioDevice(SubsystemRoot->GetAudioDeviceID());
+		return AudioDeviceManager->GetAudioDevice(DeviceId);
 	}
 
 	return FAudioDeviceHandle();
diff --git a/Engine/Source/Runtime/Engine/Public/Subsystems/AudioEngineSubsystem.h b/Engine/Source/Runtime/Engine/Public/Subsystems/AudioEngineSubsystem.h
--- a/Engine/Source/Runtime/Engine/Public/Subsystems/AudioEngineSubsystem.h
+++ b/Engine/Source/Runtime/Engine/Public/Subsystems/AudioEngineSubsystem.h
@@ -56,6 +56,9 @@ public:
 	 */
 	virtual void Update() {}
 
+	/** Returns the owning audio device Id */
+	ENGINE_API virtual Audio::FDeviceId GetAudioDeviceId() const final;
+
 	/** Returns the owning audio device handle */
 	ENGINE_API virtual FAudioDeviceHandle GetAudioDeviceHandle() const final;
 
