Crashfix for UAudioGameplayVolumeSubsystem::RemoveVolumeComponent
+ Remove unsafe references to UAudioGameplayVolumeProxy on Audio Thread

#localization none
#tests local testing, hotfixable off

diff --git a/Engine/Plugins/AudioGameplayVolume/Source/AudioGameplayVolume/Private/AudioGameplayVolumeSubsystem.cpp b/Engine/Plugins/AudioGameplayVolume/Source/AudioGameplayVolume/Private/AudioGameplayVolumeSubsystem.cpp
--- a/Engine/Plugins/AudioGameplayVolume/Source/AudioGameplayVolume/Private/AudioGameplayVolumeSubsystem.cpp
+++ b/Engine/Plugins/AudioGameplayVolume/Source/AudioGameplayVolume/Private/AudioGameplayVolumeSubsystem.cpp
@@ -2,6 +2,7 @@
 
 #include "AudioGameplayVolumeSubsystem.h"
 #include "ActiveSound.h"
+#include "Algo/Transform.h"
 #include "AudioGameplayVolumeLogs.h"
 #include "AudioGameplayVolumeComponent.h"
 #include "AudioGameplayVolumeProxy.h"
@@ -72,7 +73,7 @@ void FAudioGameplayActiveSoundInfo::Update(double ListenerInteriorStartTime)
 	InteriorSettings.UpdateInteriorValues();
 }
 
-void FAudioProxyMutatorSearchObject::SearchVolumes(const TArray<TWeakObjectPtr<UAudioGameplayVolumeProxy>>& ProxyVolumes, FAudioProxyMutatorSearchResult& OutResult)
+void FAudioProxyMutatorSearchObject::SearchVolumes(const TArray<UAudioGameplayVolumeProxy*>& ProxyVolumes, FAudioProxyMutatorSearchResult& OutResult)
 {
 	check(IsInAudioThread());
 	SCOPED_NAMED_EVENT(FAudioProxyMutatorSearchObject_SearchVolumes, FColor::Blue);
@@ -82,9 +83,9 @@ void FAudioProxyMutatorSearchObject::SearchVolumes(const TArray<TWeakObjectPtr<U
 	MutatorPriorities.PayloadType = PayloadType;
 	MutatorPriorities.bFilterPayload = bFilterPayload;
 
-	for (const TWeakObjectPtr<UAudioGameplayVolumeProxy>& ProxyVolume : ProxyVolumes)
+	for (const UAudioGameplayVolumeProxy* ProxyVolume : ProxyVolumes)
 	{
-		if (!ProxyVolume.IsValid() || ProxyVolume->GetWorldID() != WorldID)
+		if (!ProxyVolume || ProxyVolume->GetWorldID() != WorldID)
 		{
 			continue;
 		}
@@ -115,9 +116,9 @@ void FAudioProxyMutatorSearchObject::SearchVolumes(const TArray<TWeakObjectPtr<U
 
 	if (bCollectMutators)
 	{
-		for (const TWeakObjectPtr<UAudioGameplayVolumeProxy>& ProxyVolume : ProxyVolumes)
+		for (const UAudioGameplayVolumeProxy* ProxyVolume : ProxyVolumes)
 		{
-			if (!ProxyVolume.IsValid())
+			if (!ProxyVolume)
 			{
 				continue;
 			}
@@ -177,7 +178,7 @@ void UAudioGameplayVolumeSubsystem::Update()
 	check(IsInAudioThread());
 	
 	// We track and check our previous proxy count to ensure we get one update after all of the volumes have been removed
-	if (AudioGameplayVolumeConsoleVariables::bEnabled == 0 || (ProxyVolumes.Num() == 0 && PreviousProxyCount == 0))
+	if (AudioGameplayVolumeConsoleVariables::bEnabled == 0 || (KnownProxyIDs.Num() == 0 && PreviousProxyCount == 0))
 	{
 		return;
 	}
@@ -209,7 +210,7 @@ void UAudioGameplayVolumeSubsystem::Update()
 		UpdateFromListeners();
 	}
 
-	PreviousProxyCount = ProxyVolumes.Num();
+	PreviousProxyCount = KnownProxyIDs.Num();
 }
 
 void UAudioGameplayVolumeSubsystem::GatherInteriorData(const FActiveSound& ActiveSound, FSoundParseParameters& ParseParams)
@@ -228,13 +229,17 @@ void UAudioGameplayVolumeSubsystem::GatherInteriorData(const FActiveSound& Activ
 	MutatorSearch.AudioDeviceHandle = GetAudioDeviceHandle();
 	MutatorSearch.bAffectedByLegacySystem = (ActiveSound.AudioVolumeID != 0);
 
+	GenerateVolumeProxyList();
 	FAudioProxyMutatorSearchResult Result;
-	MutatorSearch.SearchVolumes(ProxyVolumes, Result);
+	MutatorSearch.SearchVolumes(TransientProxyList, Result);
 
 	// Save info about this active sound for application.
 	FAudioGameplayActiveSoundInfo& ActiveSoundInfo = ActiveSoundData.FindOrAdd(ActiveSound.GetInstanceID());
 	ActiveSoundInfo.CurrentMutators = MoveTemp(Result.MatchingMutators);
 	ActiveSoundInfo.InteriorSettings.Apply(Result.InteriorSettings);
+
+	// We need to drop references, but we can at least prevent some allocs from using a member variable
+	TransientProxyList.Reset();
 }
 
 void UAudioGameplayVolumeSubsystem::ApplyInteriorSettings(const FActiveSound& ActiveSound, FSoundParseParameters& ParseParams)
@@ -319,17 +324,24 @@ void UAudioGameplayVolumeSubsystem::AddVolumeComponent(const UAudioGameplayVolum
 	{
 		VolumeProxy->InitFromComponent(VolumeComponent);
 		AGVComponents.Emplace(ComponentID, VolumeComponent);
+		const uint32 WorldID = VolumeProxy->GetWorldID();
 
 		UE_LOG(AudioGameplayVolumeLog, VeryVerbose, TEXT("AudioGameplayVolumeComponent %s [%08x] added"), *VolumeComponent->GetFName().ToString(), ComponentID);
 
 		// Copy representation of volume to audio thread
-		TWeakObjectPtr<UAudioGameplayVolumeProxy> WeakProxy(VolumeProxy);
-		TWeakObjectPtr<UAudioGameplayVolumeSubsystem> WeakThis(this);
-		FAudioThread::RunCommandOnAudioThread([WeakThis, WeakProxy]()
+		Audio::FDeviceId CurrentDeviceId = GetAudioDeviceId();
+		FAudioThread::RunCommandOnAudioThread([CurrentDeviceId, ComponentID, WorldID]()
 		{
-			if (WeakThis.IsValid() && WeakProxy.IsValid())
+			if (FAudioDeviceManager* AudioDeviceManager = FAudioDeviceManager::Get())
 			{
-				WeakThis->AddProxy(WeakProxy);
+				FAudioDeviceHandle DeviceHandle = AudioDeviceManager->GetAudioDevice(CurrentDeviceId);
+				if (DeviceHandle.IsValid())
+				{
+					if (UAudioGameplayVolumeSubsystem* AGVSubsystem = DeviceHandle->GetSubsystem<UAudioGameplayVolumeSubsystem>())
+					{
+						AGVSubsystem->AddProxy(ComponentID, WorldID);
+					}
+				}
 			}
 		});
 	}
@@ -351,13 +363,20 @@ void UAudioGameplayVolumeSubsystem::UpdateVolumeComponent(const UAudioGameplayVo
 
 		UE_LOG(AudioGameplayVolumeLog, VeryVerbose, TEXT("AudioGameplayVolumeComponent %s [%08x] updated"), *VolumeComponent->GetFName().ToString(), ComponentID);
 
-		TWeakObjectPtr<UAudioGameplayVolumeProxy> WeakProxy(VolumeProxy);
-		TWeakObjectPtr<UAudioGameplayVolumeSubsystem> WeakThis(this);
-		FAudioThread::RunCommandOnAudioThread([WeakThis, WeakProxy]()
+		// Update representation of volume on audio thread
+		Audio::FDeviceId CurrentDeviceId = GetAudioDeviceId();
+		FAudioThread::RunCommandOnAudioThread([CurrentDeviceId, ComponentID]()
 		{
-			if (WeakThis.IsValid() && WeakProxy.IsValid())
+			if (FAudioDeviceManager* AudioDeviceManager = FAudioDeviceManager::Get())
 			{
-				WeakThis->UpdateProxy(WeakProxy);
+				FAudioDeviceHandle DeviceHandle = AudioDeviceManager->GetAudioDevice(CurrentDeviceId);
+				if (DeviceHandle.IsValid())
+				{
+					if (UAudioGameplayVolumeSubsystem* AGVSubsystem = DeviceHandle->GetSubsystem<UAudioGameplayVolumeSubsystem>())
+					{
+						AGVSubsystem->UpdateProxy(ComponentID);
+					}
+				}
 			}
 		});
 	}
@@ -419,42 +438,35 @@ bool UAudioGameplayVolumeSubsystem::DoesSupportWorld(UWorld* World) const
 	return false;
 }
 
-bool UAudioGameplayVolumeSubsystem::AddProxy(TWeakObjectPtr<UAudioGameplayVolumeProxy> WeakProxy)
+bool UAudioGameplayVolumeSubsystem::AddProxy(uint32 AudioGameplayVolumeID, uint32 WorldID)
 {
 	check(IsInAudioThread());
-	check(WeakProxy.IsValid());
 
-	// Only add to the Proxy array if the element doesn't already exist
-	const uint32 NewVolumeID = WeakProxy->GetVolumeID();
-	if (!ProxyVolumes.FindByPredicate([NewVolumeID](const TWeakObjectPtr<UAudioGameplayVolumeProxy>& Proxy) { return Proxy.IsValid() && Proxy->GetVolumeID() == NewVolumeID; }))
+	if (!AGVComponents.Contains(AudioGameplayVolumeID) || AGVComponents[AudioGameplayVolumeID]->GetProxy() == nullptr)
 	{
-		ProxyVolumes.Emplace(WeakProxy);
-
-		WorldProxyLists.FindOrAdd(WeakProxy->GetWorldID());
-		bHasStaleProxy = true;
+		return false;
+	}
 
-		UE_LOG(AudioGameplayVolumeLog, VeryVerbose, TEXT("Proxy [%08x] added"), NewVolumeID);
-		return true;
+	if (KnownProxyIDs.Contains(AudioGameplayVolumeID))
+	{
+		UE_LOG(AudioGameplayVolumeLog, VeryVerbose, TEXT("Attempting to add Proxy [%08x] multiple times"), AudioGameplayVolumeID);
+		return false;
 	}
 
-	UE_LOG(AudioGameplayVolumeLog, VeryVerbose, TEXT("Attempting to add Proxy [%08x] multiple times"), NewVolumeID);
-	return false;
+	KnownProxyIDs.Add(AudioGameplayVolumeID);
+	WorldProxyLists.FindOrAdd(WorldID);
+	bHasStaleProxy = true;
+
+	UE_LOG(AudioGameplayVolumeLog, VeryVerbose, TEXT("Proxy [%08x] added"), AudioGameplayVolumeID);
+	return true;
 }
 
-bool UAudioGameplayVolumeSubsystem::UpdateProxy(TWeakObjectPtr<UAudioGameplayVolumeProxy> WeakProxy)
+bool UAudioGameplayVolumeSubsystem::UpdateProxy(uint32 AudioGameplayVolumeID)
 {
 	check(IsInAudioThread());
-	check(WeakProxy.IsValid());
-
-	const uint32 ProxyVolumeID = WeakProxy->GetVolumeID();
-	TWeakObjectPtr<UAudioGameplayVolumeProxy>* ProxyPtr = ProxyVolumes.FindByPredicate([ProxyVolumeID](const TWeakObjectPtr<UAudioGameplayVolumeProxy>& Proxy)
-	{
-		return Proxy.IsValid() && Proxy->GetVolumeID() == ProxyVolumeID;
-	});
 
-	if (ProxyPtr)
+	if (KnownProxyIDs.Contains(AudioGameplayVolumeID))
 	{
-		*ProxyPtr = WeakProxy;
 		bHasStaleProxy = true;
 		return true;
 	}
@@ -465,11 +477,8 @@ bool UAudioGameplayVolumeSubsystem::UpdateProxy(TWeakObjectPtr<UAudioGameplayVol
 bool UAudioGameplayVolumeSubsystem::RemoveProxy(uint32 AudioGameplayVolumeID)
 {
 	check(IsInAudioThread());
-	const int32 NumRemoved = ProxyVolumes.RemoveAllSwap([AudioGameplayVolumeID](const TWeakObjectPtr<UAudioGameplayVolumeProxy>& Proxy)
-	{
-		return !Proxy.IsValid() || Proxy->GetVolumeID() == AudioGameplayVolumeID;
-	});
 
+	const int32 NumRemoved = KnownProxyIDs.Remove(AudioGameplayVolumeID);
 	if (NumRemoved > 0)
 	{
 		bHasStaleProxy = true;
@@ -480,6 +489,15 @@ bool UAudioGameplayVolumeSubsystem::RemoveProxy(uint32 AudioGameplayVolumeID)
 	return false;
 }
 
+void UAudioGameplayVolumeSubsystem::GenerateVolumeProxyList()
+{
+	TransientProxyList.Reset(256);
+	Algo::Transform(AGVComponents, TransientProxyList, [](const TPair<uint32, TObjectPtr<const UAudioGameplayVolumeComponent>>& Pair)
+	{
+		return Pair.Value ? Pair.Value->GetProxy() : nullptr;
+	});
+}
+
 bool UAudioGameplayVolumeSubsystem::IsAnyListenerInVolume(uint32 WorldID, uint32 VolumeID) const
 {
 	// We test this by checking to see if the volume id provided is in our world's current proxy list
@@ -542,6 +560,8 @@ void UAudioGameplayVolumeSubsystem::UpdateFromListeners()
 	FAudioProxyMutatorSearchResult Result;
 	TSet<uint32> TempVolumeSet;
 	FTransform ListenerTransform;
+	
+	GenerateVolumeProxyList();
 
 	// Grabbing the listeners directly here should be removed when possible - done out of necessity due to the legacy audio volume system
 	const TArray<FListener>& AudioListeners = DeviceHandle->GetListeners();
@@ -559,7 +579,7 @@ void UAudioGameplayVolumeSubsystem::UpdateFromListeners()
 				ProxySearch.bFilterPayload = false;
 				ProxySearch.bCollectMutators = false;
 				ProxySearch.bGetDefaultAudioSettings = false;
-				ProxySearch.SearchVolumes(ProxyVolumes, Result);
+				ProxySearch.SearchVolumes(TransientProxyList, Result);
 
 				// Hold on to these
 				Swap(TempVolumeSet, Result.VolumeSet);
@@ -568,7 +588,7 @@ void UAudioGameplayVolumeSubsystem::UpdateFromListeners()
 				ProxySearch.bFilterPayload = true;
 				ProxySearch.bCollectMutators = true;
 				ProxySearch.bGetDefaultAudioSettings = true;
-				ProxySearch.SearchVolumes(ProxyVolumes, Result);
+				ProxySearch.SearchVolumes(TransientProxyList, Result);
 
 				if (i < AudioListeners.Num())
 				{
@@ -593,6 +613,9 @@ void UAudioGameplayVolumeSubsystem::UpdateFromListeners()
 		AGVListeners[i].Update(Result, FVector::ZeroVector, AudioDeviceID);
 	}
 
+	// We need to drop references, but we can at least prevent some allocs from using a member variable
+	TransientProxyList.Reset();
+
 	FAudioGameplayProxyUpdateResult ProxyUpdateResult;
 	TWeakObjectPtr<UAudioGameplayVolumeSubsystem> WeakThis(this);
 
diff --git a/Engine/Plugins/AudioGameplayVolume/Source/AudioGameplayVolume/Public/AudioGameplayVolumeSubsystem.h b/Engine/Plugins/AudioGameplayVolume/Source/AudioGameplayVolume/Public/AudioGameplayVolumeSubsystem.h
--- a/Engine/Plugins/AudioGameplayVolume/Source/AudioGameplayVolume/Public/AudioGameplayVolumeSubsystem.h
+++ b/Engine/Plugins/AudioGameplayVolume/Source/AudioGameplayVolume/Public/AudioGameplayVolumeSubsystem.h
@@ -82,7 +82,7 @@ struct FAudioProxyMutatorSearchObject
 	bool bCollectMutators = true;
 	bool bGetDefaultAudioSettings = true;
 
-	void SearchVolumes(const TArray<TWeakObjectPtr<UAudioGameplayVolumeProxy>>& ProxyVolumes, FAudioProxyMutatorSearchResult& OutResult);
+	void SearchVolumes(const TArray<UAudioGameplayVolumeProxy*>& ProxyVolumes, FAudioProxyMutatorSearchResult& OutResult);
 };
 
 /**
@@ -154,8 +154,8 @@ protected:
 	bool DoesSupportWorld(UWorld* World) const;
 
 	/** (Audio Thread Only) Add, Update, Remove ProxyVolumes */
-	bool AddProxy(TWeakObjectPtr<UAudioGameplayVolumeProxy> WeakProxy);
-	bool UpdateProxy(TWeakObjectPtr<UAudioGameplayVolumeProxy> WeakProxy);
+	bool AddProxy(uint32 AudioGameplayVolumeID, uint32 WorldID);
+	bool UpdateProxy(uint32 AudioGameplayVolumeID);
 	bool RemoveProxy(uint32 AudioGameplayVolumeID);
 
 	/** Returns true if a listener associated with WorldID is inside the volume (by ID) */
@@ -167,16 +167,22 @@ protected:
 	/** Update our representation of audio listeners on the audio thread */
 	void UpdateFromListeners();
 
+	/** Generate a list of volume proxies from our component list */
+	void GenerateVolumeProxyList();
+
+	UPROPERTY(Transient)
+	TArray<TObjectPtr<UAudioGameplayVolumeProxy>> TransientProxyList;
+
 	// Components in our system
 	UPROPERTY(Transient)
 	TMap<uint32, TObjectPtr<const UAudioGameplayVolumeComponent>> AGVComponents;
 
+	// List of Component IDs (for proxies) the audio thread is tracking
+	TArray<uint32> KnownProxyIDs;
+
 	// Audio thread representation of Listeners
 	TArray<FAudioGameplayVolumeListener> AGVListeners;
 
-	// Audio thread representation of Volumes
-	TArray<TWeakObjectPtr<UAudioGameplayVolumeProxy>> ProxyVolumes;
-
 	// A collection of data about currently playing active sounds, indexed by the sound's unique ID
 	TMap<uint32, FAudioGameplayActiveSoundInfo> ActiveSoundData;
 
