Exposing component visualizers to AGV volume components

#rnx

diff --git a/Engine/Plugins/AudioGameplayVolume/AudioGameplayVolume.uplugin b/Engine/Plugins/AudioGameplayVolume/AudioGameplayVolume.uplugin
--- a/Engine/Plugins/AudioGameplayVolume/AudioGameplayVolume.uplugin
+++ b/Engine/Plugins/AudioGameplayVolume/AudioGameplayVolume.uplugin
@@ -22,7 +22,7 @@
 		{
 			"Name": "AudioGameplayVolumeEditor",
 			"Type": "Editor",
-			"LoadingPhase": "PostDefault"
+			"LoadingPhase": "PostEngineInit"
 		}
 	],
 	"Plugins" : [
diff --git a/Engine/Plugins/AudioGameplayVolume/Source/AudioGameplayVolume/Public/AudioGameplayVolumeProxy.h b/Engine/Plugins/AudioGameplayVolume/Source/AudioGameplayVolume/Public/AudioGameplayVolumeProxy.h
--- a/Engine/Plugins/AudioGameplayVolume/Source/AudioGameplayVolume/Public/AudioGameplayVolumeProxy.h
+++ b/Engine/Plugins/AudioGameplayVolume/Source/AudioGameplayVolume/Public/AudioGameplayVolumeProxy.h
@@ -6,12 +6,15 @@
 #include "AudioGameplayVolumeProxy.generated.h"
 
 // Forward Declarations 
-struct FBodyInstance;
-struct FAudioProxyMutatorPriorities;
-struct FAudioProxyMutatorSearchResult;
+class FPrimitiveDrawInterface;
 class FProxyVolumeMutator;
-class UPrimitiveComponent;
+class FSceneView;
+class UActorComponent;
 class UAudioGameplayVolumeComponent;
+class UPrimitiveComponent;
+struct FAudioProxyMutatorPriorities;
+struct FAudioProxyMutatorSearchResult;
+struct FBodyInstance;
 
 /**
  *  UAudioGameplayVolumeProxy - Abstract proxy used on audio thread to represent audio gameplay volumes.
@@ -40,6 +43,9 @@ public:
 	uint32 GetVolumeID() const;
 	uint32 GetWorldID() const;
 
+	/** Used for debug visualization of UAudioGameplayVolumeProxy in the editor */
+	virtual void DrawVisualization(const UActorComponent* Component, const FSceneView* View, FPrimitiveDrawInterface* PDI) {}
+
 protected:
 
 	ProxyMutatorList ProxyVolumeMutators;
diff --git a/Engine/Plugins/AudioGameplayVolume/Source/AudioGameplayVolumeEditor/Private/AudioGameplayVolumeComponentVisualizers.cpp b/Engine/Plugins/AudioGameplayVolume/Source/AudioGameplayVolumeEditor/Private/AudioGameplayVolumeComponentVisualizers.cpp
new file mode 100644
--- /dev/null
+++ b/Engine/Plugins/AudioGameplayVolume/Source/AudioGameplayVolumeEditor/Private/AudioGameplayVolumeComponentVisualizers.cpp
@@ -0,0 +1,20 @@
+// Copyright Epic Games, Inc. All Rights Reserved.
+
+#include "AudioGameplayVolumeComponentVisualizers.h"
+
+#include "AudioGameplayVolumeComponent.h"
+#include "AudioGameplayVolumeProxy.h"
+
+void FAudioGameplayVolumeComponentVisualizer::DrawVisualization(const UActorComponent* Component, const FSceneView* View, FPrimitiveDrawInterface* PDI)
+{
+	if (const UAudioGameplayVolumeComponent* AGVComponent = Cast<const UAudioGameplayVolumeComponent>(Component))
+	{
+		UAudioGameplayVolumeProxy* Proxy = AGVComponent->GetProxy();
+		if (!AGVComponent->IsActive() || Proxy == nullptr)
+		{
+			return;
+		}
+
+		Proxy->DrawVisualization(Component, View, PDI);
+	}
+}
diff --git a/Engine/Plugins/AudioGameplayVolume/Source/AudioGameplayVolumeEditor/Private/AudioGameplayVolumeEditorModule.cpp b/Engine/Plugins/AudioGameplayVolume/Source/AudioGameplayVolumeEditor/Private/AudioGameplayVolumeEditorModule.cpp
--- a/Engine/Plugins/AudioGameplayVolume/Source/AudioGameplayVolumeEditor/Private/AudioGameplayVolumeEditorModule.cpp
+++ b/Engine/Plugins/AudioGameplayVolume/Source/AudioGameplayVolumeEditor/Private/AudioGameplayVolumeEditorModule.cpp
@@ -1,28 +1,34 @@
 // Copyright Epic Games, Inc. All Rights Reserved.
 
 #include "AudioGameplayVolumeEditorModule.h"
+#include "AudioGameplayVolumeComponent.h"
+#include "AudioGameplayVolumeComponentVisualizers.h"
+#include "DetailCustomizations/ReverbVolumeComponentDetail.h"
+#include "Editor/UnrealEdEngine.h"
 #include "Modules/ModuleManager.h"
 #include "PropertyEditorModule.h"
-#include "DetailCustomizations/ReverbVolumeComponentDetail.h"
+#include "UnrealEdGlobals.h"
 
 void FAudioGameplayVolumeEditorModule::StartupModule()
 {
-	constexpr bool bRegisterLayouts = true;
-	HandleCustomPropertyLayouts(bRegisterLayouts);
+	constexpr bool bRegister = true;
+	HandleCustomPropertyLayouts(bRegister);
+	HandleComponentVisualizers(bRegister);
 }
 
 void FAudioGameplayVolumeEditorModule::ShutdownModule()
 {
-	constexpr bool bRegisterLayouts = false;
-	HandleCustomPropertyLayouts(bRegisterLayouts);
+	constexpr bool bRegister = false;
+	HandleCustomPropertyLayouts(bRegister);
+	HandleComponentVisualizers(bRegister);
 }
 
-void FAudioGameplayVolumeEditorModule::HandleCustomPropertyLayouts(bool bRegisterLayouts)
+void FAudioGameplayVolumeEditorModule::HandleCustomPropertyLayouts(bool bRegister)
 {
 	// Register detail customizations
 	FPropertyEditorModule& PropertyModule = FModuleManager::LoadModuleChecked<FPropertyEditorModule>(TEXT("PropertyEditor"));
 
-	if (bRegisterLayouts)
+	if (bRegister)
 	{
 		PropertyModule.RegisterCustomClassLayout("ReverbVolumeComponent", FOnGetDetailCustomizationInstance::CreateStatic(&FReverbVolumeComponentDetail::MakeInstance));
 	}
@@ -32,4 +38,38 @@ void FAudioGameplayVolumeEditorModule::HandleCustomPropertyLayouts(bool bRegiste
 	}
 }
 
+void FAudioGameplayVolumeEditorModule::HandleComponentVisualizers(bool bRegister)
+{
+	if (!GUnrealEd)
+	{
+		return;
+	}
+
+	if (bRegister)
+	{
+		RegisterComponentVisualizer(UAudioGameplayVolumeComponent::StaticClass()->GetFName(), MakeShared<FAudioGameplayVolumeComponentVisualizer>());
+	}
+	else
+	{
+		for (FName ClassName : RegisteredComponentVisualizers)
+		{
+			GUnrealEd->UnregisterComponentVisualizer(ClassName);
+		}
+
+		RegisteredComponentVisualizers.Reset(0);
+	}
+}
+
+void FAudioGameplayVolumeEditorModule::RegisterComponentVisualizer(FName ComponentClassName, TSharedPtr<FComponentVisualizer> Visualizer)
+{
+	check(GUnrealEd);
+	GUnrealEd->RegisterComponentVisualizer(ComponentClassName, Visualizer);
+	RegisteredComponentVisualizers.Add(ComponentClassName);
+
+	if (Visualizer.IsValid())
+	{
+		Visualizer->OnRegister();
+	}
+}
+
 IMPLEMENT_MODULE(FAudioGameplayVolumeEditorModule, AudioGameplayVolumeEditor);
diff --git a/Engine/Plugins/AudioGameplayVolume/Source/AudioGameplayVolumeEditor/Public/AudioGameplayVolumeComponentVisualizers.h b/Engine/Plugins/AudioGameplayVolume/Source/AudioGameplayVolumeEditor/Public/AudioGameplayVolumeComponentVisualizers.h
new file mode 100644
--- /dev/null
+++ b/Engine/Plugins/AudioGameplayVolume/Source/AudioGameplayVolumeEditor/Public/AudioGameplayVolumeComponentVisualizers.h
@@ -0,0 +1,22 @@
+// Copyright Epic Games, Inc. All Rights Reserved.
+
+#pragma once
+
+#include "ComponentVisualizer.h"
+
+// Forward Declarations 
+class UActorComponent;
+class FSceneView;
+class FPrimitiveDrawInterface;
+
+class FAudioGameplayVolumeComponentVisualizer : public FComponentVisualizer
+{
+public:
+
+	FAudioGameplayVolumeComponentVisualizer() = default;
+	virtual ~FAudioGameplayVolumeComponentVisualizer() = default;
+
+	//~ Begin FComponentVisualizer Interface
+	virtual void DrawVisualization(const UActorComponent* Component, const FSceneView* View, FPrimitiveDrawInterface* PDI) override;
+	//~ End FComponentVisualizer Interface
+};
diff --git a/Engine/Plugins/AudioGameplayVolume/Source/AudioGameplayVolumeEditor/Public/AudioGameplayVolumeEditorModule.h b/Engine/Plugins/AudioGameplayVolume/Source/AudioGameplayVolumeEditor/Public/AudioGameplayVolumeEditorModule.h
--- a/Engine/Plugins/AudioGameplayVolume/Source/AudioGameplayVolumeEditor/Public/AudioGameplayVolumeEditorModule.h
+++ b/Engine/Plugins/AudioGameplayVolume/Source/AudioGameplayVolumeEditor/Public/AudioGameplayVolumeEditorModule.h
@@ -2,8 +2,12 @@
 
 #pragma once
 
-#include "Modules/ModuleInterface.h"
 #include "Logging/LogMacros.h"
+#include "Modules/ModuleInterface.h"
+#include "Templates/SharedPointerFwd.h"
+
+// Forward Declarations 
+class FComponentVisualizer;
 
 AUDIOGAMEPLAYVOLUMEEDITOR_API DECLARE_LOG_CATEGORY_EXTERN(AudioGameplayVolumeEditor, Log, All);
 
@@ -16,5 +20,10 @@ public:
 
 protected:
 
-	void HandleCustomPropertyLayouts(bool bRegisterLayouts);
+	void HandleCustomPropertyLayouts(bool bRegister);
+	void HandleComponentVisualizers(bool bRegister);
+
+	void RegisterComponentVisualizer(FName ComponentClassName, TSharedPtr<FComponentVisualizer> Visualizer);
+
+	TArray<FName> RegisteredComponentVisualizers;
 };
