Fix unsafe pointer usage in AudioGameplayVolumeProxy

#rnx

diff --git a/Engine/Plugins/AudioGameplayVolume/Source/AudioGameplayVolume/Private/AudioGameplayVolumeProxy.cpp b/Engine/Plugins/AudioGameplayVolume/Source/AudioGameplayVolume/Private/AudioGameplayVolumeProxy.cpp
--- a/Engine/Plugins/AudioGameplayVolume/Source/AudioGameplayVolume/Private/AudioGameplayVolumeProxy.cpp
+++ b/Engine/Plugins/AudioGameplayVolume/Source/AudioGameplayVolume/Private/AudioGameplayVolumeProxy.cpp
@@ -19,6 +19,13 @@ namespace AudioGameplayVolumeConsoleVariables
 		TEXT("Skips physics body queries for proxies that are not close to the listener.\n0: Disable, 1: Enable (default)"),
 		ECVF_Default);
 
+	int32 ProxyReferenceType = 1;
+	FAutoConsoleVariableRef CVarProxyReferenceType(
+		TEXT("au.AudioGameplayVolumes.ProxyReferenceType"),
+		ProxyReferenceType,
+		TEXT("Denotes the type of reference the proxy object holds.\n0: WeakObject, 1: TObject uproperty (default)"),
+		ECVF_Default);
+
 } // namespace AudioGameplayVolumeConsoleVariables
 
 #include UE_INLINE_GENERATED_CPP_BY_NAME(AudioGameplayVolumeProxy)
@@ -131,7 +138,18 @@ bool UAGVPrimitiveComponentProxy::ContainsPosition(const FVector& Position) cons
 	SCOPED_NAMED_EVENT(UAGVPrimitiveComponentProxy_ContainsPosition, FColor::Blue);
 
 	FBodyInstance* BodyInstancePointer = nullptr;
-	if (UPrimitiveComponent* PrimitiveComponent = WeakPrimative.Get())
+	UPrimitiveComponent* PrimitiveComponent = nullptr;
+
+	if (AudioGameplayVolumeConsoleVariables::ProxyReferenceType != 0)
+	{
+		PrimitiveComponent = PrimitiveComp;
+	}
+	else
+	{
+		PrimitiveComponent = WeakPrimative.Get();
+	}
+	
+	if (PrimitiveComponent)
 	{
 		if (NeedsPhysicsQuery(PrimitiveComponent, Position))
 		{
@@ -165,7 +183,14 @@ void UAGVPrimitiveComponentProxy::InitFromComponent(const UAudioGameplayVolumeCo
 
 		if (PrimitiveComponents.Num() > 0)
 		{
-			WeakPrimative = PrimitiveComponents[0];
+			if (AudioGameplayVolumeConsoleVariables::ProxyReferenceType != 0)
+			{
+				PrimitiveComp = PrimitiveComponents[0];
+			}
+			else
+			{
+				WeakPrimative = PrimitiveComponents[0];
+			}
 		}
 	}
 }
@@ -201,7 +226,16 @@ bool UAGVConditionProxy::ContainsPosition(const FVector& Position) const
 {
 	SCOPED_NAMED_EVENT(UAGVConditionProxy_ContainsPosition, FColor::Blue);
 
-	const UObject* ObjectWithInterface = WeakObject.Get();
+	const UObject* ObjectWithInterface = nullptr;
+	if (AudioGameplayVolumeConsoleVariables::ProxyReferenceType != 0)
+	{
+		ObjectWithInterface = ObjectPtr;
+	}
+	else
+	{
+		ObjectWithInterface = WeakObject.Get();
+	}
+
 	if (ObjectWithInterface && ObjectWithInterface->Implements<UAudioGameplayCondition>())
 	{
 		return IAudioGameplayCondition::Execute_ConditionMet(ObjectWithInterface)
@@ -220,7 +254,14 @@ void UAGVConditionProxy::InitFromComponent(const UAudioGameplayVolumeComponent*
 	{
 		if (OwnerActor->Implements<UAudioGameplayCondition>())
 		{
-			WeakObject = MakeWeakObjectPtr(OwnerActor);
+			if (AudioGameplayVolumeConsoleVariables::ProxyReferenceType != 0)
+			{
+				ObjectPtr = OwnerActor;
+			}
+			else
+			{
+				WeakObject = MakeWeakObjectPtr(OwnerActor);
+			}
 		}
 		else
 		{
@@ -230,7 +271,14 @@ void UAGVConditionProxy::InitFromComponent(const UAudioGameplayVolumeComponent*
 			{
 				if (ActorComponent && ActorComponent->Implements<UAudioGameplayCondition>())
 				{
-					WeakObject = MakeWeakObjectPtr(ActorComponent);
+					if (AudioGameplayVolumeConsoleVariables::ProxyReferenceType != 0)
+					{
+						ObjectPtr = ActorComponent;
+					}
+					else
+					{
+						WeakObject = MakeWeakObjectPtr(ActorComponent);
+					}
 					break;
 				}
 			}
diff --git a/Engine/Plugins/AudioGameplayVolume/Source/AudioGameplayVolume/Public/AudioGameplayVolumeProxy.h b/Engine/Plugins/AudioGameplayVolume/Source/AudioGameplayVolume/Public/AudioGameplayVolumeProxy.h
--- a/Engine/Plugins/AudioGameplayVolume/Source/AudioGameplayVolume/Public/AudioGameplayVolumeProxy.h
+++ b/Engine/Plugins/AudioGameplayVolume/Source/AudioGameplayVolume/Public/AudioGameplayVolumeProxy.h
@@ -75,6 +75,9 @@ protected:
 	/** Returns false if we can skip the physics body query, which can be expensive */
 	bool NeedsPhysicsQuery(UPrimitiveComponent* PrimitiveComponent, const FVector& Position) const;
 
+	UPROPERTY(Transient)
+	TObjectPtr<UPrimitiveComponent> PrimitiveComp;
+
 	TWeakObjectPtr<UPrimitiveComponent> WeakPrimative;
 };
 
@@ -95,5 +98,8 @@ public:
 
 protected:
 
+	UPROPERTY(Transient)
+	TObjectPtr<const UObject> ObjectPtr;
+
 	TWeakObjectPtr<const UObject> WeakObject;
 };
