Hide Resonance Effect Plugin from AGV Reverb Volume component settings

#rnx

diff --git a/Engine/Plugins/AudioGameplayVolume/Source/AudioGameplayVolumeEditor/Private/DetailCustomizations/ReverbVolumeComponentDetail.cpp b/Engine/Plugins/AudioGameplayVolume/Source/AudioGameplayVolumeEditor/Private/DetailCustomizations/ReverbVolumeComponentDetail.cpp
--- a/Engine/Plugins/AudioGameplayVolume/Source/AudioGameplayVolumeEditor/Private/DetailCustomizations/ReverbVolumeComponentDetail.cpp
+++ b/Engine/Plugins/AudioGameplayVolume/Source/AudioGameplayVolumeEditor/Private/DetailCustomizations/ReverbVolumeComponentDetail.cpp
@@ -54,7 +54,8 @@ void FReverbVolumeComponentDetail::CustomizeDetails(IDetailLayoutBuilder& Detail
 		const FProperty* ChildProperty = ChildPropertyHandle->GetProperty();
 		const FName ChildPropertyName = ChildProperty->GetFName();
 
-		if (ChildPropertyName != GET_MEMBER_NAME_CHECKED(FReverbSettings, bApplyReverb))
+		if (ChildPropertyName != GET_MEMBER_NAME_CHECKED(FReverbSettings, bApplyReverb) &&
+			ChildPropertyName != GET_MEMBER_NAME_CHECKED(FReverbSettings, ReverbPluginEffect))
 		{
 			ReverbCategory.AddProperty(ChildPropertyHandle);
 		}
