[Quail] Fix crash in AudioGameplayVolume::AddProxy

#rnx

diff --git a/Engine/Plugins/AudioGameplayVolume/Source/AudioGameplayVolume/Private/AudioGameplayVolumeSubsystem.cpp b/Engine/Plugins/AudioGameplayVolume/Source/AudioGameplayVolume/Private/AudioGameplayVolumeSubsystem.cpp
--- a/Engine/Plugins/AudioGameplayVolume/Source/AudioGameplayVolume/Private/AudioGameplayVolumeSubsystem.cpp
+++ b/Engine/Plugins/AudioGameplayVolume/Source/AudioGameplayVolume/Private/AudioGameplayVolumeSubsystem.cpp
@@ -442,7 +442,14 @@ bool UAudioGameplayVolumeSubsystem::AddProxy(uint32 AudioGameplayVolumeID, uint3
 {
 	check(IsInAudioThread());
 
-	if (!AGVComponents.Contains(AudioGameplayVolumeID) || AGVComponents[AudioGameplayVolumeID]->GetProxy() == nullptr)
+	// Make sure the component is in the map and non-null
+	if (!AGVComponents.Contains(AudioGameplayVolumeID) || AGVComponents[AudioGameplayVolumeID] == nullptr)
+	{
+		return false;
+	}
+
+	// Make sure our proxy is valid
+	if (AGVComponents[AudioGameplayVolumeID]->GetProxy() == nullptr)
 	{
 		return false;
 	}
