[AudioGameplayVolumes] Allowing Primitive proxies to support multiple primitives.   (Pull request)

#pr 12636
#rnx

diff --git a/Engine/Plugins/AudioGameplayVolume/Source/AudioGameplayVolume/Private/AudioGameplayVolumeProxy.cpp b/Engine/Plugins/AudioGameplayVolume/Source/AudioGameplayVolume/Private/AudioGameplayVolumeProxy.cpp
--- a/Engine/Plugins/AudioGameplayVolume/Source/AudioGameplayVolume/Private/AudioGameplayVolumeProxy.cpp
+++ b/Engine/Plugins/AudioGameplayVolume/Source/AudioGameplayVolume/Private/AudioGameplayVolumeProxy.cpp
@@ -133,38 +133,77 @@ UAGVPrimitiveComponentProxy::UAGVPrimitiveComponentProxy(const FObjectInitialize
 {
 }
 
-bool UAGVPrimitiveComponentProxy::ContainsPosition(const FVector& Position) const
+bool NeedsPhysicsQuery(UPrimitiveComponent* PrimitiveComponent, const FVector& Position)
 {
-	SCOPED_NAMED_EVENT(UAGVPrimitiveComponentProxy_ContainsPosition, FColor::Blue);
-
-	FBodyInstance* BodyInstancePointer = nullptr;
-	UPrimitiveComponent* PrimitiveComponent = nullptr;
+	check(PrimitiveComponent);
 
-	if (AudioGameplayVolumeConsoleVariables::ProxyReferenceType != 0)
+	if (!PrimitiveComponent->IsPhysicsStateCreated() || !PrimitiveComponent->HasValidPhysicsState())
 	{
-		PrimitiveComponent = PrimitiveComp;
+		return false;
 	}
-	else
+
+	// Temporary kill switch for distance culling
+	if (AudioGameplayVolumeConsoleVariables::bProxyDistanceCulling == 0)
 	{
-		PrimitiveComponent = WeakPrimative.Get();
+		return true;
 	}
-	
-	if (PrimitiveComponent)
+
+	// Early distance culling
+	const float BoundsRadiusSq = FMath::Square(PrimitiveComponent->Bounds.SphereRadius);
+	const float DistanceSq = FVector::DistSquared(PrimitiveComponent->Bounds.Origin, Position);
+
+	return DistanceSq <= BoundsRadiusSq;
+}
+
+static UPrimitiveComponent* GetPrimitiveComponent(const TObjectPtr<UPrimitiveComponent>& PrimitiveComponent)
+{
+	return PrimitiveComponent;
+}
+
+static UPrimitiveComponent* GetPrimitiveComponent(const TWeakObjectPtr<UPrimitiveComponent>& PrimitiveComponent)
+{
+	return PrimitiveComponent.Get();
+}
+
+template<typename ComponentType>
+static bool ContainsPositionHelper(const TArray<ComponentType>& Primitives, const FVector& Position)
+{
+	bool ContainsPosition = false;
+
+	for (const ComponentType& Primitive : Primitives)
 	{
-		if (NeedsPhysicsQuery(PrimitiveComponent, Position))
+		UPrimitiveComponent* PrimitiveComponent = GetPrimitiveComponent(Primitive);
+		if (PrimitiveComponent)
 		{
-			BodyInstancePointer = PrimitiveComponent->GetBodyInstance();
+			if (NeedsPhysicsQuery(PrimitiveComponent, Position) && PrimitiveComponent->GetBodyInstance())
+			{
+				const FBodyInstance* BodyInstance = PrimitiveComponent->GetBodyInstance();
+				float DistanceSquared = 0.f;
+				FVector PointOnBody = FVector::ZeroVector;
+				if (BodyInstance->GetSquaredDistanceToBody(Position, DistanceSquared, PointOnBody) && FMath::IsNearlyZero(DistanceSquared))
+				{
+					ContainsPosition = true;
+					break;
+				}
+			}
 		}
 	}
 
-	if (!BodyInstancePointer)
+	return ContainsPosition;
+}
+
+bool UAGVPrimitiveComponentProxy::ContainsPosition(const FVector& Position) const
+{
+	SCOPED_NAMED_EVENT(UAGVPrimitiveComponentProxy_ContainsPosition, FColor::Blue);
+
+	if (AudioGameplayVolumeConsoleVariables::ProxyReferenceType != 0)
 	{
-		return false;
+		return ContainsPositionHelper(Primitives, Position);
+	}
+	else
+	{
+		return ContainsPositionHelper(WeakPrimitives, Position);
 	}
-
-	float DistanceSquared = 0.f;
-	FVector PointOnBody = FVector::ZeroVector;
-	return BodyInstancePointer->GetSquaredDistanceToBody(Position, DistanceSquared, PointOnBody) && FMath::IsNearlyZero(DistanceSquared);
 }
 
 void UAGVPrimitiveComponentProxy::InitFromComponent(const UAudioGameplayVolumeComponent* Component)
@@ -174,49 +213,23 @@ void UAGVPrimitiveComponentProxy::InitFromComponent(const UAudioGameplayVolumeCo
 	if (Component)
 	{
 		TInlineComponentArray<UPrimitiveComponent*> PrimitiveComponents(Component->GetOwner());
-		const int32 PrimitiveComponentCount = PrimitiveComponents.Num();
-
-		if (PrimitiveComponentCount != 1)
-		{
-			UE_LOG(AudioGameplayVolumeLog, Warning, TEXT("Was expecting exactly one Primitive Component on the owning actor, found %d - this could cause unexpected behavior"), PrimitiveComponentCount);
-		}
 
 		if (PrimitiveComponents.Num() > 0)
 		{
 			if (AudioGameplayVolumeConsoleVariables::ProxyReferenceType != 0)
 			{
-				PrimitiveComp = PrimitiveComponents[0];
+				Primitives.Reserve(PrimitiveComponents.Num());
+				Algo::Copy(PrimitiveComponents, Primitives);
 			}
 			else
 			{
-				WeakPrimative = PrimitiveComponents[0];
+				WeakPrimitives.Reserve(PrimitiveComponents.Num());
+				Algo::Copy(PrimitiveComponents, WeakPrimitives);
 			}
 		}
 	}
 }
 
-bool UAGVPrimitiveComponentProxy::NeedsPhysicsQuery(UPrimitiveComponent* PrimitiveComponent, const FVector& Position) const
-{
-	check(PrimitiveComponent);
-
-	if (!PrimitiveComponent->IsPhysicsStateCreated() || !PrimitiveComponent->HasValidPhysicsState())
-	{
-		return false;
-	}
-
-	// Temporary kill switch for distance culling
-	if (AudioGameplayVolumeConsoleVariables::bProxyDistanceCulling == 0)
-	{
-		return true;
-	}
-
-	// Early distance culling
-	const float BoundsRadiusSq = FMath::Square(PrimitiveComponent->Bounds.SphereRadius);
-	const float DistanceSq = FVector::DistSquared(PrimitiveComponent->Bounds.Origin, Position);
-
-	return DistanceSq <= BoundsRadiusSq;
-}
-
 UAGVConditionProxy::UAGVConditionProxy(const FObjectInitializer& ObjectInitializer)
 	: Super(ObjectInitializer)
 {
diff --git a/Engine/Plugins/AudioGameplayVolume/Source/AudioGameplayVolume/Public/AudioGameplayVolumeProxy.h b/Engine/Plugins/AudioGameplayVolume/Source/AudioGameplayVolume/Public/AudioGameplayVolumeProxy.h
--- a/Engine/Plugins/AudioGameplayVolume/Source/AudioGameplayVolume/Public/AudioGameplayVolumeProxy.h
+++ b/Engine/Plugins/AudioGameplayVolume/Source/AudioGameplayVolume/Public/AudioGameplayVolumeProxy.h
@@ -71,14 +71,11 @@ public:
 	virtual void InitFromComponent(const UAudioGameplayVolumeComponent* Component) override;
 
 protected:
-	
-	/** Returns false if we can skip the physics body query, which can be expensive */
-	bool NeedsPhysicsQuery(UPrimitiveComponent* PrimitiveComponent, const FVector& Position) const;
 
 	UPROPERTY(Transient)
-	TObjectPtr<UPrimitiveComponent> PrimitiveComp;
+	TArray<TObjectPtr<UPrimitiveComponent>> Primitives;
 
-	TWeakObjectPtr<UPrimitiveComponent> WeakPrimative;
+	TArray<TWeakObjectPtr<UPrimitiveComponent>> WeakPrimitives;
 };
 
 /**
