Removing weak pointers in AGV proxies and wrapping new multiple primitive functionality in a cvar as a temporary fallback

#rnx

diff --git a/Engine/Plugins/AudioGameplayVolume/Source/AudioGameplayVolume/Private/AudioGameplayVolumeProxy.cpp b/Engine/Plugins/AudioGameplayVolume/Source/AudioGameplayVolume/Private/AudioGameplayVolumeProxy.cpp
--- a/Engine/Plugins/AudioGameplayVolume/Source/AudioGameplayVolume/Private/AudioGameplayVolumeProxy.cpp
+++ b/Engine/Plugins/AudioGameplayVolume/Source/AudioGameplayVolume/Private/AudioGameplayVolumeProxy.cpp
@@ -19,17 +19,45 @@ namespace AudioGameplayVolumeConsoleVariables
 		TEXT("Skips physics body queries for proxies that are not close to the listener.\n0: Disable, 1: Enable (default)"),
 		ECVF_Default);
 
-	int32 ProxyReferenceType = 1;
-	FAutoConsoleVariableRef CVarProxyReferenceType(
-		TEXT("au.AudioGameplayVolumes.ProxyReferenceType"),
-		ProxyReferenceType,
-		TEXT("Denotes the type of reference the proxy object holds.\n0: WeakObject, 1: TObject uproperty (default)"),
+	int32 bAllowMultiplePrimitives = 1;
+	FAutoConsoleVariableRef CVarAllowMultiplePrimitives (
+		TEXT("au.AudioGameplayVolumes.AllowMultiplePrimitives"),
+		bAllowMultiplePrimitives,
+		TEXT("Allows consideration of all the primitives on an actor for intersection tests.\n")
+		TEXT("NOTE: The option to turn this off will be removed in the future. Use with caution!\n")
+		TEXT("0: Off, 1: On (default)"),
 		ECVF_Default);
 
 } // namespace AudioGameplayVolumeConsoleVariables
 
 #include UE_INLINE_GENERATED_CPP_BY_NAME(AudioGameplayVolumeProxy)
 
+namespace AudioGameplayVolumeUtils
+{
+	bool NeedsPhysicsQuery(UPrimitiveComponent* PrimitiveComponent, const FVector& Position)
+	{
+		check(PrimitiveComponent);
+
+		if (!PrimitiveComponent->IsPhysicsStateCreated() || !PrimitiveComponent->HasValidPhysicsState())
+		{
+			return false;
+		}
+
+		// Temporary kill switch for distance culling
+		if (AudioGameplayVolumeConsoleVariables::bProxyDistanceCulling == 0)
+		{
+			return true;
+		}
+
+		// Early distance culling
+		const float BoundsRadiusSq = FMath::Square(PrimitiveComponent->Bounds.SphereRadius);
+		const float DistanceSq = FVector::DistSquared(PrimitiveComponent->Bounds.Origin, Position);
+
+		return DistanceSq <= BoundsRadiusSq;
+	}
+
+} // namespace AudioGameplayVolumeUtils
+
 UAudioGameplayVolumeProxy::UAudioGameplayVolumeProxy(const FObjectInitializer& ObjectInitializer)
 	: Super(ObjectInitializer)
 {
@@ -133,58 +161,33 @@ UAGVPrimitiveComponentProxy::UAGVPrimitiveComponentProxy(const FObjectInitialize
 {
 }
 
-bool NeedsPhysicsQuery(UPrimitiveComponent* PrimitiveComponent, const FVector& Position)
+bool UAGVPrimitiveComponentProxy::ContainsPosition(const FVector& Position) const
 {
-	check(PrimitiveComponent);
+	SCOPED_NAMED_EVENT(UAGVPrimitiveComponentProxy_ContainsPosition, FColor::Blue);
 
-	if (!PrimitiveComponent->IsPhysicsStateCreated() || !PrimitiveComponent->HasValidPhysicsState())
-	{
-		return false;
-	}
+	bool ContainsPosition = false;
 
-	// Temporary kill switch for distance culling
-	if (AudioGameplayVolumeConsoleVariables::bProxyDistanceCulling == 0)
+	for (const TObjectPtr<UPrimitiveComponent>& PrimitiveComponent : Primitives)
 	{
-		return true;
-	}
-
-	// Early distance culling
-	const float BoundsRadiusSq = FMath::Square(PrimitiveComponent->Bounds.SphereRadius);
-	const float DistanceSq = FVector::DistSquared(PrimitiveComponent->Bounds.Origin, Position);
-
-	return DistanceSq <= BoundsRadiusSq;
-}
-
-static UPrimitiveComponent* GetPrimitiveComponent(const TObjectPtr<UPrimitiveComponent>& PrimitiveComponent)
-{
-	return PrimitiveComponent;
-}
-
-static UPrimitiveComponent* GetPrimitiveComponent(const TWeakObjectPtr<UPrimitiveComponent>& PrimitiveComponent)
-{
-	return PrimitiveComponent.Get();
-}
+		if (!PrimitiveComponent)
+		{
+			continue;
+		}
 
-template<typename ComponentType>
-static bool ContainsPositionHelper(const TArray<ComponentType>& Primitives, const FVector& Position)
-{
-	bool ContainsPosition = false;
+		const FBodyInstance* BodyInstance = PrimitiveComponent->GetBodyInstance();
+		if (!BodyInstance)
+		{
+			continue;
+		}
 
-	for (const ComponentType& Primitive : Primitives)
-	{
-		UPrimitiveComponent* PrimitiveComponent = GetPrimitiveComponent(Primitive);
-		if (PrimitiveComponent)
+		if (AudioGameplayVolumeUtils::NeedsPhysicsQuery(PrimitiveComponent, Position))
 		{
-			if (NeedsPhysicsQuery(PrimitiveComponent, Position) && PrimitiveComponent->GetBodyInstance())
+			float DistanceSquared = 0.f;
+			FVector PointOnBody = FVector::ZeroVector;
+			if (BodyInstance->GetSquaredDistanceToBody(Position, DistanceSquared, PointOnBody) && FMath::IsNearlyZero(DistanceSquared))
 			{
-				const FBodyInstance* BodyInstance = PrimitiveComponent->GetBodyInstance();
-				float DistanceSquared = 0.f;
-				FVector PointOnBody = FVector::ZeroVector;
-				if (BodyInstance->GetSquaredDistanceToBody(Position, DistanceSquared, PointOnBody) && FMath::IsNearlyZero(DistanceSquared))
-				{
-					ContainsPosition = true;
-					break;
-				}
+				ContainsPosition = true;
+				break;
 			}
 		}
 	}
@@ -192,39 +195,28 @@ static bool ContainsPositionHelper(const TArray<ComponentType>& Primitives, cons
 	return ContainsPosition;
 }
 
-bool UAGVPrimitiveComponentProxy::ContainsPosition(const FVector& Position) const
-{
-	SCOPED_NAMED_EVENT(UAGVPrimitiveComponentProxy_ContainsPosition, FColor::Blue);
-
-	if (AudioGameplayVolumeConsoleVariables::ProxyReferenceType != 0)
-	{
-		return ContainsPositionHelper(Primitives, Position);
-	}
-	else
-	{
-		return ContainsPositionHelper(WeakPrimitives, Position);
-	}
-}
-
 void UAGVPrimitiveComponentProxy::InitFromComponent(const UAudioGameplayVolumeComponent* Component)
 {
 	Super::InitFromComponent(Component);
+	Primitives.Reset();
 
 	if (Component)
 	{
 		TInlineComponentArray<UPrimitiveComponent*> PrimitiveComponents(Component->GetOwner());
+		int32 PrimitiveCount = PrimitiveComponents.Num();
 
-		if (PrimitiveComponents.Num() > 0)
+		if (PrimitiveCount > 0)
 		{
-			if (AudioGameplayVolumeConsoleVariables::ProxyReferenceType != 0)
+			// Using this console variable as a temporary rollback to give us time to respond in live development.
+			// Not expecting this rollback to live long as the (original) case it's protecting was non-deterministic
+			if (!AudioGameplayVolumeConsoleVariables::bAllowMultiplePrimitives)
 			{
-				Primitives.Reserve(PrimitiveComponents.Num());
-				Algo::Copy(PrimitiveComponents, Primitives);
+				Primitives.Emplace(PrimitiveComponents[0]);
 			}
 			else
 			{
-				WeakPrimitives.Reserve(PrimitiveComponents.Num());
-				Algo::Copy(PrimitiveComponents, WeakPrimitives);
+				Primitives.Reserve(PrimitiveCount);
+				Algo::Copy(PrimitiveComponents, Primitives);
 			}
 		}
 	}
@@ -239,20 +231,10 @@ bool UAGVConditionProxy::ContainsPosition(const FVector& Position) const
 {
 	SCOPED_NAMED_EVENT(UAGVConditionProxy_ContainsPosition, FColor::Blue);
 
-	const UObject* ObjectWithInterface = nullptr;
-	if (AudioGameplayVolumeConsoleVariables::ProxyReferenceType != 0)
-	{
-		ObjectWithInterface = ObjectPtr;
-	}
-	else
+	if (ObjectPtr && ObjectPtr->Implements<UAudioGameplayCondition>())
 	{
-		ObjectWithInterface = WeakObject.Get();
-	}
-
-	if (ObjectWithInterface && ObjectWithInterface->Implements<UAudioGameplayCondition>())
-	{
-		return IAudioGameplayCondition::Execute_ConditionMet(ObjectWithInterface)
-			|| IAudioGameplayCondition::Execute_ConditionMet_Position(ObjectWithInterface, Position);
+		return IAudioGameplayCondition::Execute_ConditionMet(ObjectPtr)
+			|| IAudioGameplayCondition::Execute_ConditionMet_Position(ObjectPtr, Position);
 	}
 
 	return false;
@@ -267,14 +249,7 @@ void UAGVConditionProxy::InitFromComponent(const UAudioGameplayVolumeComponent*
 	{
 		if (OwnerActor->Implements<UAudioGameplayCondition>())
 		{
-			if (AudioGameplayVolumeConsoleVariables::ProxyReferenceType != 0)
-			{
-				ObjectPtr = OwnerActor;
-			}
-			else
-			{
-				WeakObject = MakeWeakObjectPtr(OwnerActor);
-			}
+			ObjectPtr = OwnerActor;
 		}
 		else
 		{
@@ -284,14 +259,7 @@ void UAGVConditionProxy::InitFromComponent(const UAudioGameplayVolumeComponent*
 			{
 				if (ActorComponent && ActorComponent->Implements<UAudioGameplayCondition>())
 				{
-					if (AudioGameplayVolumeConsoleVariables::ProxyReferenceType != 0)
-					{
-						ObjectPtr = ActorComponent;
-					}
-					else
-					{
-						WeakObject = MakeWeakObjectPtr(ActorComponent);
-					}
+					ObjectPtr = ActorComponent;
 					break;
 				}
 			}
diff --git a/Engine/Plugins/AudioGameplayVolume/Source/AudioGameplayVolume/Public/AudioGameplayVolumeProxy.h b/Engine/Plugins/AudioGameplayVolume/Source/AudioGameplayVolume/Public/AudioGameplayVolumeProxy.h
--- a/Engine/Plugins/AudioGameplayVolume/Source/AudioGameplayVolume/Public/AudioGameplayVolumeProxy.h
+++ b/Engine/Plugins/AudioGameplayVolume/Source/AudioGameplayVolume/Public/AudioGameplayVolumeProxy.h
@@ -74,8 +74,6 @@ protected:
 
 	UPROPERTY(Transient)
 	TArray<TObjectPtr<UPrimitiveComponent>> Primitives;
-
-	TArray<TWeakObjectPtr<UPrimitiveComponent>> WeakPrimitives;
 };
 
 /**
@@ -97,6 +95,4 @@ protected:
 
 	UPROPERTY(Transient)
 	TObjectPtr<const UObject> ObjectPtr;
-
-	TWeakObjectPtr<const UObject> WeakObject;
 };
