Avoid queueing tasks for UpdateComponentsFromProxyInfo if they won't be doing any work.

From replayrun, saves 10us gamethread time per frame on xb1. Two audio commands were being sent per frame.

diff --git a/Engine/Plugins/AudioGameplayVolume/Source/AudioGameplayVolume/Private/AudioGameplayVolumeSubsystem.cpp b/Engine/Plugins/AudioGameplayVolume/Source/AudioGameplayVolume/Private/AudioGameplayVolumeSubsystem.cpp
--- a/Engine/Plugins/AudioGameplayVolume/Source/AudioGameplayVolume/Private/AudioGameplayVolumeSubsystem.cpp
+++ b/Engine/Plugins/AudioGameplayVolume/Source/AudioGameplayVolume/Private/AudioGameplayVolumeSubsystem.cpp
@@ -10,6 +10,7 @@
 #include "AudioDevice.h"
 #include "Engine/World.h"
 #include "Misc/App.h"
+#include "Stats/Stats.h"
 
 #include UE_INLINE_GENERATED_CPP_BY_NAME(AudioGameplayVolumeSubsystem)
 
@@ -629,13 +630,17 @@ void UAudioGameplayVolumeSubsystem::UpdateFromListeners()
 	for (TPair<uint32, FAudioGameplayVolumeProxyInfo>& WorldProxyInfo : WorldProxyLists)
 	{
 		WorldProxyInfo.Value.Update(AGVListeners, ProxyUpdateResult);
-		FAudioThread::RunCommandOnGameThread([WeakThis, ProxyUpdateResult]
+		if (!ProxyUpdateResult.EnteredProxies.IsEmpty() || !ProxyUpdateResult.ExitedProxies.IsEmpty())
 		{
-			if (WeakThis.IsValid())
+			FAudioThread::RunCommandOnGameThread([WeakThis, ProxyUpdateResult]
 			{
-				WeakThis->UpdateComponentsFromProxyInfo(ProxyUpdateResult);
-			}
-		});
+				QUICK_SCOPE_CYCLE_COUNTER(STAT_AudioCommand_UpdateComponentsFromProxyInfo);
+				if (WeakThis.IsValid())
+				{
+					WeakThis->UpdateComponentsFromProxyInfo(ProxyUpdateResult);
+				}
+			});
+		}
 	}
 }
 
diff --git a/Engine/Plugins/AudioGameplayVolume/Source/AudioGameplayVolume/Public/AudioGameplayVolumeSubsystem.h b/Engine/Plugins/AudioGameplayVolume/Source/AudioGameplayVolume/Public/AudioGameplayVolumeSubsystem.h
--- a/Engine/Plugins/AudioGameplayVolume/Source/AudioGameplayVolume/Public/AudioGameplayVolumeSubsystem.h
+++ b/Engine/Plugins/AudioGameplayVolume/Source/AudioGameplayVolume/Public/AudioGameplayVolumeSubsystem.h
@@ -43,7 +43,6 @@ struct FAudioGameplayProxyUpdateResult
 {
 	TSet<uint32> EnteredProxies;
 	TSet<uint32> ExitedProxies;
-	bool bForceUpdate = false;
 };
 
 /**
