Prevent cluster verification to fail by duplicating volume proxy if spawned actor references an existing volume proxy. Code is guarded with cvar (enabled by default) au.AudioGameplayVolumes.DuplicateProxy. Courtesy of nathan.walker

#rnx

diff --git a/Engine/Plugins/AudioGameplayVolume/Source/AudioGameplayVolume/Private/AudioGameplayVolumeComponent.cpp b/Engine/Plugins/AudioGameplayVolume/Source/AudioGameplayVolume/Private/AudioGameplayVolumeComponent.cpp
--- a/Engine/Plugins/AudioGameplayVolume/Source/AudioGameplayVolume/Private/AudioGameplayVolumeComponent.cpp
+++ b/Engine/Plugins/AudioGameplayVolume/Source/AudioGameplayVolume/Private/AudioGameplayVolumeComponent.cpp
@@ -1,12 +1,23 @@
 // Copyright Epic Games, Inc. All Rights Reserved.
 
 #include "AudioGameplayVolumeComponent.h"
+#include "AudioGameplayVolumeProxy.h"
 #include "AudioGameplayVolumeSubsystem.h"
 #include "AudioDevice.h"
 #include "Engine/World.h"
 
 #include UE_INLINE_GENERATED_CPP_BY_NAME(AudioGameplayVolumeComponent)
 
+namespace AudioGameplayVolumeComponentVariables
+{
+	int32 bDuplicateProxy = 1;
+	FAutoConsoleVariableRef CVarProxyDistanceCulling(
+		TEXT("au.AudioGameplayVolumes.DuplicateProxy"),
+		bDuplicateProxy,
+		TEXT("Prevent cluster verification to fail by duplicating volume proxy if spawned actor references an existing volume proxy - 0: Disable, 1: Enable (default)"),
+		ECVF_Default);
+}
+
 UAudioGameplayVolumeComponent::UAudioGameplayVolumeComponent(const FObjectInitializer& ObjectInitializer)
 	: Super(ObjectInitializer)
 {
@@ -77,6 +88,19 @@ void UAudioGameplayVolumeComponent::PostEditChangeProperty(FPropertyChangedEvent
 }
 #endif // WITH_EDITOR
 
+void UAudioGameplayVolumeComponent::OnRegister()
+{
+	Super::OnRegister();
+
+	if (AudioGameplayVolumeComponentVariables::bDuplicateProxy)
+	{
+		if (Proxy != nullptr && Proxy->IsTemplate())
+		{
+			Proxy = DuplicateObject<UAudioGameplayVolumeProxy>(Proxy, this);
+		}
+	}
+}
+
 void UAudioGameplayVolumeComponent::OnUnregister()
 {
 	Super::OnUnregister();
diff --git a/Engine/Plugins/AudioGameplayVolume/Source/AudioGameplayVolume/Public/AudioGameplayVolumeComponent.h b/Engine/Plugins/AudioGameplayVolume/Source/AudioGameplayVolume/Public/AudioGameplayVolumeComponent.h
--- a/Engine/Plugins/AudioGameplayVolume/Source/AudioGameplayVolume/Public/AudioGameplayVolumeComponent.h
+++ b/Engine/Plugins/AudioGameplayVolume/Source/AudioGameplayVolume/Public/AudioGameplayVolumeComponent.h
@@ -56,6 +56,7 @@ protected:
 #endif // WITH_EDITOR
 
 	//~ Begin UActorComponent Interface
+	virtual void OnRegister() override;
 	virtual void OnUnregister() override;
 	//~ End UActorComponent Interface
 
