AGV Subsystem: add cvar to opt into behavior that prevents duplicate submix sends from being added to a sound

diff --git a/Engine/Plugins/AudioGameplayVolume/Source/AudioGameplayVolume/Private/AudioGameplayVolumeSubsystem.cpp b/Engine/Plugins/AudioGameplayVolume/Source/AudioGameplayVolume/Private/AudioGameplayVolumeSubsystem.cpp
--- a/Engine/Plugins/AudioGameplayVolume/Source/AudioGameplayVolume/Private/AudioGameplayVolumeSubsystem.cpp
+++ b/Engine/Plugins/AudioGameplayVolume/Source/AudioGameplayVolume/Private/AudioGameplayVolumeSubsystem.cpp
@@ -51,6 +51,13 @@ namespace AudioGameplayVolumeConsoleVariables
 		bAudioThreadCmdRollback,
 		TEXT("When non-zero, uses old code to rollback late thread command change."),
 		ECVF_Default);
+	
+	int32 bPreventDuplicateSubmixSends = 0;
+	FAutoConsoleVariableRef CVarPreventDuplicateSubmixSends(
+		TEXT("au.AudioGameplayVolumes.PreventDuplicateSubmixSends"),
+		bPreventDuplicateSubmixSends,
+		TEXT("When non-zero, will not add a Submix Sends if one already is there sending to the same Submix."),
+		ECVF_Default);
 
 } // namespace AudioGameplayVolumeConsoleVariables
 
@@ -291,8 +298,39 @@ void UAudioGameplayVolumeSubsystem::ApplyInteriorSettings(const FActiveSound& Ac
 	ActiveSoundInfo->CurrentInteriorVolume = Params.SourceInteriorVolume;
 	ActiveSoundInfo->CurrentInteriorLPF = Params.SourceInteriorLPF;
 
-	// Apply to our parse params
-	ParseParams.SoundSubmixSends.Append(Params.SoundSubmixSends);
+	// Apply Submix Sends to our parse param
+	if (AudioGameplayVolumeConsoleVariables::bPreventDuplicateSubmixSends)
+	{
+		// This opt-in behavior prevents the AudioGameplayVolumeSubsystem from clobbering existing Submix Sends
+		// that may already exist on the asset we're interacting with.
+		// Reason: If a sound has two Submix Sends to the same Submix, the current observed behavior is that only the last send is respected.
+		// Example: A designer explicitly dialed in a Send to Submix A at 0.5f
+		// CVAR OFF: this subsystem will clobber that send
+		// CVAR  ON: this subsystem will preserve that design decision and no-op for that particular send
+		bool bSendAlreadyExists = false;
+		for (const FSoundSubmixSendInfo& SendA : Params.SoundSubmixSends)
+		{
+			for (const FSoundSubmixSendInfo& SendB : ParseParams.SoundSubmixSends)
+			{
+				if (SendA.SoundSubmix == SendB.SoundSubmix)
+				{
+					bSendAlreadyExists = true;
+					break;
+				}
+			}
+		}
+
+		if (!bSendAlreadyExists)
+		{
+			ParseParams.SoundSubmixSends.Append(Params.SoundSubmixSends);
+		}
+	}
+	else
+	{
+		ParseParams.SoundSubmixSends.Append(Params.SoundSubmixSends);
+	}
+
+	// Apply to the rest of our parse params
 	ParseParams.InteriorVolumeMultiplier = ActiveSoundInfo->CurrentInteriorVolume;
 	ParseParams.AmbientZoneFilterFrequency = ActiveSoundInfo->CurrentInteriorLPF;
 }
