AGV Updates / Glow up
+ Batching notifications for entering & exiting proxies
+ Updated tooltips for better documentation of features
+ Renaming (DisplayNames) to be more intuitive for the end user
+ Volume Proxy -> Audio Toggle
+ AGV Primitive Proxy -> Audio Listener In Primitives
+ AGV Proximity Proxy (internal) -> Audio Listener In Range
+ AGV Condition Proxy -> Arbitrary
+ AudioGameplayVolumeInteraction (interface) -> Audio Toggle Interaction
+ Attenuation (component) -> Interior-Exterior Attenuation

#wip
#rnx

diff --git a/Engine/Plugins/AudioGameplay/Source/AudioGameplay/Public/Interfaces/IAudioGameplayCondition.h b/Engine/Plugins/AudioGameplay/Source/AudioGameplay/Public/Interfaces/IAudioGameplayCondition.h
--- a/Engine/Plugins/AudioGameplay/Source/AudioGameplay/Public/Interfaces/IAudioGameplayCondition.h
+++ b/Engine/Plugins/AudioGameplay/Source/AudioGameplay/Public/Interfaces/IAudioGameplayCondition.h
@@ -5,7 +5,17 @@
 #include "UObject/Interface.h"
 #include "IAudioGameplayCondition.generated.h"
 
-/** Interface for audio gameplay conditions */
+/**
+ * Interface for implementing conditions for the audio toggle system.
+ *
+ * Make a component or actor implement this interface to provide its Audio Toggle component
+ * set in Arbitrary condition mode a way to compute its On/Off state.
+ * If the actor implements it, the Audio Toggle will abide by the implementation on the actor.
+ * If any component on the same actor implements it, the Audio Toggle will abide by the
+ * implementation on the *first* of these components.
+ *
+ * See: Audio Toggle component (UAudioGameplayVolumeComponent)
+ */
 UINTERFACE(BlueprintType, MinimalAPI)
 class UAudioGameplayCondition : public UInterface
 {
@@ -19,9 +29,10 @@ class IAudioGameplayCondition
 public:
 
 	/**
-	 * Basic condition check
+	 * Generic condition check, whatever the toggle is being evaluated for (Audio Listener,
+	 * Active sound,...)
 	 * 
-	 * @returns true if condition is met
+	 * @returns true if toggle state should be On, false otherwise
 	 */
 	UFUNCTION(BlueprintCallable, BlueprintNativeEvent, Category = "AudioGameplayCondition")
 	AUDIOGAMEPLAY_API bool ConditionMet() const;
@@ -29,10 +40,12 @@ public:
 	
 	/**
 	 * Allows testing a condition against a provided position.
+	 * This position can be that of the Audio Listener, of the Active Sound,... for which the
+	 * toggle is being evaluated.
 	 *
 	 * @param Position - The location to be considered
 	 * 
-	 * @returns true if condition is met
+	 * @returns true if toggle state should be On for this entity's position, false otherwise
 	 */
 	UFUNCTION(BlueprintCallable, BlueprintNativeEvent, Category = "AudioGameplayCondition")
 	AUDIOGAMEPLAY_API bool ConditionMet_Position(const FVector& Position) const;
diff --git a/Engine/Plugins/AudioGameplay/Source/AudioGameplay/Public/Interfaces/IAudioGameplayVolumeInteraction.h b/Engine/Plugins/AudioGameplay/Source/AudioGameplay/Public/Interfaces/IAudioGameplayVolumeInteraction.h
--- a/Engine/Plugins/AudioGameplay/Source/AudioGameplay/Public/Interfaces/IAudioGameplayVolumeInteraction.h
+++ b/Engine/Plugins/AudioGameplay/Source/AudioGameplay/Public/Interfaces/IAudioGameplayVolumeInteraction.h
@@ -5,8 +5,19 @@
 #include "UObject/Interface.h"
 #include "IAudioGameplayVolumeInteraction.generated.h"
 
-/** Interface for interacting with the audio gameplay volume system */
-UINTERFACE(BlueprintType, MinimalAPI)
+/**
+ * Interface for interacting with the audio toggle system.
+ *
+ * Make a component implement this interface to receive notifications when the toggle is
+ * active for at least one listener or inactive for all.
+ *
+ * The owning actor must have an Audio Toggle component for other components of this actor
+ * implementing this interface to receive the notifications.
+ *
+ * Instead of implementing this interface, you can also just bind to the "On Toggled On" and
+ * "On Toggled Off" events of the Audio Toggle component directly.
+ */
+UINTERFACE(BlueprintType, MinimalAPI, meta = (DisplayName = "Audio Toggle Interaction"))
 class UAudioGameplayVolumeInteraction : public UInterface
 {
 	GENERATED_BODY()
@@ -19,16 +30,16 @@ class IAudioGameplayVolumeInteraction
 public:
 
 	/**
-	 * Called when a listener 'enters' the associated proxy
+	 * Called when the Audio Toggle on the same actor switches to ON.
 	 */
-	UFUNCTION(BlueprintCallable, BlueprintNativeEvent, Category = "AudioGameplayVolume")
+	UFUNCTION(BlueprintCallable, BlueprintNativeEvent, Category = "AudioGameplayVolume", meta = (DisplayName = "On Toggled On"))
 	AUDIOGAMEPLAY_API void OnListenerEnter();
 	virtual void OnListenerEnter_Implementation() {}
 
 	/**
-	 * Called when a listener 'exits' the associated proxy
+	 * Called when the Audio Toggle on the same actor switches to OFF.
 	 */
-	UFUNCTION(BlueprintCallable, BlueprintNativeEvent, Category = "AudioGameplayVolume")
+	UFUNCTION(BlueprintCallable, BlueprintNativeEvent, Category = "AudioGameplayVolume", meta = (DisplayName = "On Toggled Off"))
 	AUDIOGAMEPLAY_API void OnListenerExit();
 	virtual void OnListenerExit_Implementation() {}
 };
diff --git a/Engine/Plugins/AudioGameplayVolume/Source/AudioGameplayVolume/Private/AudioGameplayVolumeListener.cpp b/Engine/Plugins/AudioGameplayVolume/Source/AudioGameplayVolume/Private/AudioGameplayVolumeListener.cpp
--- a/Engine/Plugins/AudioGameplayVolume/Source/AudioGameplayVolume/Private/AudioGameplayVolumeListener.cpp
+++ b/Engine/Plugins/AudioGameplayVolume/Source/AudioGameplayVolume/Private/AudioGameplayVolumeListener.cpp
@@ -62,7 +62,7 @@ void FAudioGameplayVolumeListener::Update(const FAudioProxyMutatorSearchResult&
 
 	Position = InPosition;
 
-	Swap(PreviousProxies, CurrentProxies);
+	PreviousProxies = MoveTemp(CurrentProxies);
 	CurrentProxies = Result.VolumeSet;
 
 	TSet<uint32> EnteredProxies = CurrentProxies.Difference(PreviousProxies);
diff --git a/Engine/Plugins/AudioGameplayVolume/Source/AudioGameplayVolume/Private/AudioGameplayVolumeSubsystem.cpp b/Engine/Plugins/AudioGameplayVolume/Source/AudioGameplayVolume/Private/AudioGameplayVolumeSubsystem.cpp
--- a/Engine/Plugins/AudioGameplayVolume/Source/AudioGameplayVolume/Private/AudioGameplayVolumeSubsystem.cpp
+++ b/Engine/Plugins/AudioGameplayVolume/Source/AudioGameplayVolume/Private/AudioGameplayVolumeSubsystem.cpp
@@ -660,25 +660,32 @@ void UAudioGameplayVolumeSubsystem::UpdateFromListeners()
 
 	// We need to drop references, but we can at least prevent some allocs from using a member variable
 	TransientProxyList.Reset();
-
+	
+	// For all our worlds, update which proxies each listener is in and determine if there was a change
 	FAudioGameplayProxyUpdateResult ProxyUpdateResult;
-	TWeakObjectPtr<UAudioGameplayVolumeSubsystem> WeakThis(this);
-
-	// Update our world proxy lists
+	TArray<FAudioGameplayProxyUpdateResult> UpdateResultsForProxyWithPresenceChanges;
 	for (TPair<uint32, FAudioGameplayVolumeProxyInfo>& WorldProxyInfo : WorldProxyLists)
 	{
 		WorldProxyInfo.Value.Update(AGVListeners, ProxyUpdateResult);
 		if (!ProxyUpdateResult.EnteredProxies.IsEmpty() || !ProxyUpdateResult.ExitedProxies.IsEmpty())
 		{
-			FAudioThread::RunCommandOnGameThread([WeakThis, ProxyUpdateResult]
+			UpdateResultsForProxyWithPresenceChanges.Add(ProxyUpdateResult);
+		}
+	}
+
+	// If we had some changes, batch-broadcast the proxy entries/exits notifications
+	if (!UpdateResultsForProxyWithPresenceChanges.IsEmpty())
+	{
+		FAudioThread::RunCommandOnGameThread([WeakThis = TWeakObjectPtr(this), ProxyUpdateResults = MoveTemp(UpdateResultsForProxyWithPresenceChanges)]
+		{
+			QUICK_SCOPE_CYCLE_COUNTER(STAT_AudioCommand_UpdateComponentsFromProxyInfo);
+			if (ThisClass* This = WeakThis.Get())
 			{
-				QUICK_SCOPE_CYCLE_COUNTER(STAT_AudioCommand_UpdateComponentsFromProxyInfo);
-				if (WeakThis.IsValid())
+				for (const FAudioGameplayProxyUpdateResult& ProxyUpdateResult : ProxyUpdateResults)
 				{
-					WeakThis->UpdateComponentsFromProxyInfo(ProxyUpdateResult);
+					This->UpdateComponentsFromProxyInfo(ProxyUpdateResult);
 				}
-			});
-		}
+			}
+		});
 	}
 }
-
diff --git a/Engine/Plugins/AudioGameplayVolume/Source/AudioGameplayVolume/Public/AttenuationVolumeComponent.h b/Engine/Plugins/AudioGameplayVolume/Source/AudioGameplayVolume/Public/AttenuationVolumeComponent.h
--- a/Engine/Plugins/AudioGameplayVolume/Source/AudioGameplayVolume/Public/AttenuationVolumeComponent.h
+++ b/Engine/Plugins/AudioGameplayVolume/Source/AudioGameplayVolume/Public/AttenuationVolumeComponent.h
@@ -31,9 +31,11 @@ protected:
 };
 
 /**
- *  UAttenuationVolumeComponent - Audio Gameplay Volume component for occlusion settings (volume attenuation)
+ *  Component driving the volume of sounds playing inside or outside the associated AGV according to where
+ *  the audio listener lays relative to this spatial volume.
+ *  Note: This component is only compatible with Audio Toggles using a "Listener in Primitives" or "Listener in Range" condition.
  */
-UCLASS(MinimalAPI, Blueprintable, Config = Game, ClassGroup = ("AudioGameplayVolume"), meta = (BlueprintSpawnableComponent, DisplayName = "Attenuation"))
+UCLASS(MinimalAPI, Blueprintable, Config = Game, ClassGroup = ("AudioGameplayVolume"), meta = (BlueprintSpawnableComponent, DisplayName = "Interior-Exterior Attenuation"))
 class UAttenuationVolumeComponent : public UAudioGameplayVolumeMutator
 {
 	GENERATED_UCLASS_BODY()
diff --git a/Engine/Plugins/AudioGameplayVolume/Source/AudioGameplayVolume/Public/AudioGameplayVolumeComponent.h b/Engine/Plugins/AudioGameplayVolume/Source/AudioGameplayVolume/Public/AudioGameplayVolumeComponent.h
--- a/Engine/Plugins/AudioGameplayVolume/Source/AudioGameplayVolume/Public/AudioGameplayVolumeComponent.h
+++ b/Engine/Plugins/AudioGameplayVolume/Source/AudioGameplayVolume/Public/AudioGameplayVolumeComponent.h
@@ -15,10 +15,17 @@ class UAudioGameplayVolumeSubsystem;
 DECLARE_DYNAMIC_MULTICAST_DELEGATE(FOnAudioGameplayVolumeProxyStateChange);
 
 /**
- *  UAudioGameplayVolumeComponent - Component used to drive interaction with AudioGameplayVolumeSubsystem.
- *   NOTE: Do not inherit from this class, use UAudioGameplayVolumeComponentBase or UAudioGameplayVolumeMutator to create extendable functionality
+ *  Component used to drive interaction with AudioGameplayVolumeSubsystem ("Audio Toggle" subsystem).
+ *
+ *  This is the component that makes the owning actor behave like an Audio Toggle, and where you can configure how a toggle
+ *  can be active or inactive (through its Toggle Condition). Blueprints can bind arbitrary effects to the toggle becoming
+ *  On or Off through this component's "On Toggled On" and "On Toggled Off" events. But you will usually want to add
+ *  Audio
+ *  
+ *  NOTE: Do not inherit from this class, use UAudioGameplayVolumeComponentBase or UAudioGameplayVolumeMutator to create
+ *  extendable functionality.
  */
-UCLASS(MinimalAPI, Config = Game, ClassGroup = ("AudioGameplay"), meta = (BlueprintSpawnableComponent, IsBlueprintBase = false, DisplayName = "Volume Proxy"))
+UCLASS(MinimalAPI, Config = Game, ClassGroup = ("AudioGameplay"), meta = (BlueprintSpawnableComponent, IsBlueprintBase = false, DisplayName = "Audio Toggle"))
 class UAudioGameplayVolumeComponent final : public UAudioGameplayComponent
 {
 	GENERATED_UCLASS_BODY()
@@ -40,17 +47,20 @@ public:
 	UE_API void ExitProxy() const;
 
 	/** Blueprint event for proxy enter */
-	UPROPERTY(BlueprintAssignable, Category = Events)
+	UPROPERTY(BlueprintAssignable, Category = Events, meta = (DisplayName = "On Toggled On"))
 	FOnAudioGameplayVolumeProxyStateChange OnProxyEnter;
 
 	/** Blueprint event for proxy exit */
-	UPROPERTY(BlueprintAssignable, Category = Events)
+	UPROPERTY(BlueprintAssignable, Category = Events, meta = (DisplayName = "On Toggled Off"))
 	FOnAudioGameplayVolumeProxyStateChange OnProxyExit;
 
 protected:
 
-	// A representation of this volume for the audio thread
-	UPROPERTY(Instanced, EditAnywhere, BlueprintReadWrite, Category = "AudioGameplay", Meta = (ShowOnlyInnerProperties, AllowPrivateAccess = "true"))
+	/**
+	 * The kind of condition that this Audio Toggle uses to evaluate its On/Off state.
+	 * See the tooltip for each choice to see how they work and can be further configured.
+	 */
+	UPROPERTY(Instanced, EditAnywhere, BlueprintReadWrite, Category = "AudioGameplay", Meta = (ShowOnlyInnerProperties, AllowPrivateAccess = "true", DisplayName = "Toggle Condition"))
 	TObjectPtr<UAudioGameplayVolumeProxy> Proxy = nullptr;
 
 #if WITH_EDITOR
diff --git a/Engine/Plugins/AudioGameplayVolume/Source/AudioGameplayVolume/Public/AudioGameplayVolumeProxy.h b/Engine/Plugins/AudioGameplayVolume/Source/AudioGameplayVolume/Public/AudioGameplayVolumeProxy.h
--- a/Engine/Plugins/AudioGameplayVolume/Source/AudioGameplayVolume/Public/AudioGameplayVolumeProxy.h
+++ b/Engine/Plugins/AudioGameplayVolume/Source/AudioGameplayVolume/Public/AudioGameplayVolumeProxy.h
@@ -19,7 +19,11 @@ struct FAudioProxyMutatorSearchResult;
 struct FBodyInstance;
 
 /**
- *  UAudioGameplayVolumeProxy - Abstract proxy used on audio thread to represent audio gameplay volumes.
+ *  Abstract condition type for Audio Toggles (UAudioGameplayVolumeComponent).
+ *  
+ *  The way Audio Toggles evaluate to On or Off is implemented in child classes by overriding
+ *  the ContainsPosition function (because toggles were originally only implemented as the
+ *  audio listener being inside or outside some 3D area (or "volume")).
  */
 UCLASS(MinimalAPI, Abstract, EditInlineNew, HideDropdown)
 class UAudioGameplayVolumeProxy : public UObject
@@ -58,9 +62,13 @@ protected:
 };
 
 /**
- *  UAGVPrimitiveComponentProxy - Proxy based on a volume's primitive component
+ *  Sets the Audio Toggle to On if the audio listener is in any of the primitives
+ *  (as in, Primitive Components) of the owning actor, sets to Off otherwise.
+ *  
+ *  Note: The primitive components must have some physics enabled (for instance,
+ *  OverlapAll - Primitives with the NoCollision profile are ignored).
  */
-UCLASS(MinimalAPI, meta = (DisplayName = "AGV Primitive Proxy"))
+UCLASS(MinimalAPI, meta = (DisplayName = "Audio Listener in Primitives"))
 class UAGVPrimitiveComponentProxy : public UAudioGameplayVolumeProxy
 {
 	GENERATED_UCLASS_BODY()
@@ -79,9 +87,15 @@ protected:
 };
 
 /**
- *  UAGVConditionProxy - Proxy for use with the UAudioGameplayCondition interface
+ *  Sets the Audio Toggle state according to some arbitrary condition, regularly re-evaluated.
+ *  
+ *  If the owner actor implements IAudioGameplayCondition, sets the Audio Toggle to On
+ *  when either ConditionMet or ConditionMet_Position is true, and sets it to Off when
+ *  they are both false.
+ *  If the owner actor does not implement it, looks for the first component on the
+ *  actor found implementing IAudioGameplayCondition, then follows the same rules as above.
  */
-UCLASS(MinimalAPI, meta = (DisplayName = "AGV Condition Proxy"))
+UCLASS(MinimalAPI, meta = (DisplayName = "Arbitrary"))
 class UAGVConditionProxy : public UAudioGameplayVolumeProxy
 {
 	GENERATED_UCLASS_BODY()
