Simplified mutator search in AudioGameplayVolume code

#rnx

diff --git a/Engine/Plugins/AudioGameplayVolume/Source/AudioGameplayVolume/Private/AudioGameplayVolumeMutator.cpp b/Engine/Plugins/AudioGameplayVolume/Source/AudioGameplayVolume/Private/AudioGameplayVolumeMutator.cpp
--- a/Engine/Plugins/AudioGameplayVolume/Source/AudioGameplayVolume/Private/AudioGameplayVolumeMutator.cpp
+++ b/Engine/Plugins/AudioGameplayVolume/Source/AudioGameplayVolume/Private/AudioGameplayVolumeMutator.cpp
@@ -91,35 +91,6 @@ FProxyVolumeMutator::FProxyVolumeMutator()
 	MutatorName = MutatorBaseName;
 }
 
-void FProxyVolumeMutator::UpdatePriority(FAudioProxyMutatorPriorities& Priorities) const
-{
-	if (!Priorities.bFilterPayload || HasPayloadType(Priorities.PayloadType))
-	{
-		int32 DefaultValue = INDEX_NONE;
-		int32& CurrentPriority = Priorities.PriorityMap.FindOrAdd(MutatorName, DefaultValue);
-		if (Priority > CurrentPriority)
-		{
-			CurrentPriority = Priority;
-		}
-	}
-}
-
-bool FProxyVolumeMutator::CheckPriority(const FAudioProxyMutatorPriorities& Priorities) const
-{
-	if (!Priorities.bFilterPayload || HasPayloadType(Priorities.PayloadType))
-	{
-		if (const int32* HighestPriority = Priorities.PriorityMap.Find(MutatorName))
-		{
-			if (Priority == *HighestPriority)
-			{
-				return true;
-			}
-		}
-	}
-
-	return false;
-}
-
 void FProxyVolumeMutator::Apply(FInteriorSettings& InteriorSettings) const
 {
 	InteriorSettings.bIsWorldSettings = false;
diff --git a/Engine/Plugins/AudioGameplayVolume/Source/AudioGameplayVolume/Private/AudioGameplayVolumeProxy.cpp b/Engine/Plugins/AudioGameplayVolume/Source/AudioGameplayVolume/Private/AudioGameplayVolumeProxy.cpp
--- a/Engine/Plugins/AudioGameplayVolume/Source/AudioGameplayVolume/Private/AudioGameplayVolumeProxy.cpp
+++ b/Engine/Plugins/AudioGameplayVolume/Source/AudioGameplayVolume/Private/AudioGameplayVolumeProxy.cpp
@@ -104,37 +104,12 @@ void UAudioGameplayVolumeProxy::InitFromComponent(const UAudioGameplayVolumeComp
 	Audio::Analytics::RecordEvent_Usage(TEXT("AudioGameplayVolume.InitializedFromComponent"));
 }
 
-void UAudioGameplayVolumeProxy::FindMutatorPriority(FAudioProxyMutatorPriorities& Priorities) const
+void UAudioGameplayVolumeProxy::CollectMutators(ProxyMutatorList& OutMutators) const
 {
 	check(IsInAudioThread());
-	for (const TSharedPtr<FProxyVolumeMutator>& ProxyVolumeMutator : ProxyVolumeMutators)
-	{
-		if (!ProxyVolumeMutator.IsValid())
-		{
-			continue;
-		}
-
-		ProxyVolumeMutator->UpdatePriority(Priorities);
-	}
+	OutMutators.Append(ProxyVolumeMutators);
 }
 
-void UAudioGameplayVolumeProxy::GatherMutators(const FAudioProxyMutatorPriorities& Priorities, FAudioProxyMutatorSearchResult& OutResult) const
-{
-	check(IsInAudioThread());
-	for (const TSharedPtr<FProxyVolumeMutator>& ProxyVolumeMutator : ProxyVolumeMutators)
-	{
-		if (!ProxyVolumeMutator.IsValid())
-		{
-			continue;
-		}
-
-		if (ProxyVolumeMutator->CheckPriority(Priorities))
-		{
-			ProxyVolumeMutator->Apply(OutResult.InteriorSettings);
-			OutResult.MatchingMutators.Add(ProxyVolumeMutator);
-		}
-	}
-}
 
 void UAudioGameplayVolumeProxy::AddPayloadType(PayloadFlags InType)
 {
diff --git a/Engine/Plugins/AudioGameplayVolume/Source/AudioGameplayVolume/Private/AudioGameplayVolumeSubsystem.cpp b/Engine/Plugins/AudioGameplayVolume/Source/AudioGameplayVolume/Private/AudioGameplayVolumeSubsystem.cpp
--- a/Engine/Plugins/AudioGameplayVolume/Source/AudioGameplayVolume/Private/AudioGameplayVolumeSubsystem.cpp
+++ b/Engine/Plugins/AudioGameplayVolume/Source/AudioGameplayVolume/Private/AudioGameplayVolumeSubsystem.cpp
@@ -87,9 +87,6 @@ void FAudioProxyMutatorSearchObject::SearchVolumes(const TArray<UAudioGameplayVo
 	SCOPED_NAMED_EVENT(FAudioProxyMutatorSearchObject_SearchVolumes, FColor::Blue);
 
 	OutResult.Reset();
-	FAudioProxyMutatorPriorities MutatorPriorities;
-	MutatorPriorities.PayloadType = PayloadType;
-	MutatorPriorities.bFilterPayload = bFilterPayload;
 
 	for (const UAudioGameplayVolumeProxy* ProxyVolume : ProxyVolumes)
 	{
@@ -110,8 +107,7 @@ void FAudioProxyMutatorSearchObject::SearchVolumes(const TArray<UAudioGameplayVo
 
 		if (bCollectMutators)
 		{
-			// We only need to calculate priorities for mutators if we're collecting them.
-			ProxyVolume->FindMutatorPriority(MutatorPriorities);
+			ProxyVolume->CollectMutators(OutResult.MatchingMutators);
 		}
 		OutResult.VolumeSet.Add(ProxyVolume->GetVolumeID());
 	}
@@ -122,19 +118,46 @@ void FAudioProxyMutatorSearchObject::SearchVolumes(const TArray<UAudioGameplayVo
 		AudioDeviceHandle->GetDefaultAudioSettings(WorldID, OutResult.ReverbSettings, OutResult.InteriorSettings);
 	}
 
+	// Sort our mutators and apply interior settings
 	if (bCollectMutators)
 	{
-		for (const UAudioGameplayVolumeProxy* ProxyVolume : ProxyVolumes)
+		auto SortPredicate = [](TSharedPtr<FProxyVolumeMutator>& A, TSharedPtr<FProxyVolumeMutator>& B)
 		{
-			if (!ProxyVolume)
+			check(A.IsValid() && B.IsValid());
+
+			if (A->MutatorName.IsEqual(B->MutatorName))
 			{
-				continue;
+				return A->Priority > B->Priority;
+			}
+
+			return A->MutatorName.FastLess(B->MutatorName);
+		};
+
+		// Remove invalid mutators and sort
+		OutResult.MatchingMutators.RemoveAllSwap([](TSharedPtr<FProxyVolumeMutator>& ProxyMutator) { return !ProxyMutator.IsValid(); });
+		Algo::Sort(OutResult.MatchingMutators, SortPredicate);
+
+		FName MutatorGroup;
+		int32 CurrentPriority = INDEX_NONE;
+
+		// Apply mutators to our interior settings.
+		// Mutators are sorted by type (name) and descending priority, so we can quickly apply the highest priority mutators.
+		for (TSharedPtr<FProxyVolumeMutator>& Mutator : OutResult.MatchingMutators)
+		{
+			// If we encounter a new name, it's a new type - with the highest priority for that group
+			if (!Mutator->MutatorName.IsEqual(MutatorGroup))
+			{
+				MutatorGroup = Mutator->MutatorName;
+				CurrentPriority = Mutator->Priority;
 			}
 
-			if (OutResult.VolumeSet.Contains(ProxyVolume->GetVolumeID()))
+			// Lower priority mutators are ignored, we need to advance to the next named group
+			if (Mutator->Priority < CurrentPriority)
 			{
-				ProxyVolume->GatherMutators(MutatorPriorities, OutResult);
+				continue;
 			}
+
+			Mutator->Apply(OutResult.InteriorSettings);
 		}
 	}
 }
@@ -284,13 +307,32 @@ void UAudioGameplayVolumeSubsystem::ApplyInteriorSettings(const FActiveSound& Ac
 	FAudioProxyActiveSoundParams Params(*ActiveSoundInfo, Listener);
 	Params.bAllowSpatialization = ActiveSound.bAllowSpatialization;
 
-	for (const TSharedPtr<FProxyVolumeMutator>& SoundMutator : ActiveSoundInfo->CurrentMutators)
+	// Apply mutators to our active sound params
+	// Mutators are sorted by type (name) and descending priority, so we can quickly apply the highest priority mutators.
+	FName MutatorGroup;
+	int32 CurrentPriority = INDEX_NONE;
+	for (const TSharedPtr<FProxyVolumeMutator>& Mutator : ActiveSoundInfo->CurrentMutators)
 	{
-		if (SoundMutator.IsValid())
+		if (!Mutator.IsValid())
+		{
+			continue;
+		}
+
+		// If we encounter a new name, it's a new type - with the highest priority for that group
+		if (!Mutator->MutatorName.IsEqual(MutatorGroup))
 		{
-			Params.bListenerInVolume = IsAnyListenerInVolume(WorldID, SoundMutator->VolumeID);
-			SoundMutator->Apply(Params);
+			MutatorGroup = Mutator->MutatorName;
+			CurrentPriority = Mutator->Priority;
 		}
+
+		// Lower priority mutators are ignored, we need to advance to the next named group
+		if (Mutator->Priority < CurrentPriority)
+		{
+			continue;
+		}
+
+		Params.bListenerInVolume = IsAnyListenerInVolume(WorldID, Mutator->VolumeID);
+		Mutator->Apply(Params);
 	}
 
 	// Update interior values
diff --git a/Engine/Plugins/AudioGameplayVolume/Source/AudioGameplayVolume/Public/AudioGameplayVolumeMutator.h b/Engine/Plugins/AudioGameplayVolume/Source/AudioGameplayVolume/Public/AudioGameplayVolumeMutator.h
--- a/Engine/Plugins/AudioGameplayVolume/Source/AudioGameplayVolume/Public/AudioGameplayVolumeMutator.h
+++ b/Engine/Plugins/AudioGameplayVolume/Source/AudioGameplayVolume/Public/AudioGameplayVolumeMutator.h
@@ -42,18 +42,6 @@ struct FAudioProxyActiveSoundParams
 	void UpdateInteriorValues();
 };
 
-/**
- *  FAudioProxyMutatorPriorities - Used for finding the highest priority mutators on a proxy
- */
-struct FAudioProxyMutatorPriorities
-{
-	using PayloadFlags = AudioGameplay::EComponentPayload;
-
-	TMap<FName, int32> PriorityMap;
-	PayloadFlags PayloadType = PayloadFlags::AGCP_None;
-	bool bFilterPayload = true;
-};
-
 /**
  *  FProxyVolumeMutator - An audio thread representation of the payload for an AudioGameplayVolumeComponent.
  */
@@ -66,10 +54,7 @@ public:
 
 	using PayloadFlags = AudioGameplay::EComponentPayload;
 
-	UE_API virtual void UpdatePriority(FAudioProxyMutatorPriorities& Priorities) const;
-	UE_API virtual bool CheckPriority(const FAudioProxyMutatorPriorities& Priorities) const;
-
-	UE_API virtual void Apply(FInteriorSettings& InteriorSettings) const;
+	virtual void Apply(FInteriorSettings& InteriorSettings) const;
 	virtual void Apply(FAudioProxyActiveSoundParams& Params) const {}
 	virtual void Apply(FAudioGameplayVolumeListener& Listener) const {}
 
diff --git a/Engine/Plugins/AudioGameplayVolume/Source/AudioGameplayVolume/Public/AudioGameplayVolumeProxy.h b/Engine/Plugins/AudioGameplayVolume/Source/AudioGameplayVolume/Public/AudioGameplayVolumeProxy.h
--- a/Engine/Plugins/AudioGameplayVolume/Source/AudioGameplayVolume/Public/AudioGameplayVolumeProxy.h
+++ b/Engine/Plugins/AudioGameplayVolume/Source/AudioGameplayVolume/Public/AudioGameplayVolumeProxy.h
@@ -14,8 +14,6 @@ class FSceneView;
 class UActorComponent;
 class UAudioGameplayVolumeComponent;
 class UPrimitiveComponent;
-struct FAudioProxyMutatorPriorities;
-struct FAudioProxyMutatorSearchResult;
 struct FBodyInstance;
 
 /**
@@ -40,8 +38,7 @@ public:
 	UE_API virtual bool ContainsPosition(const FVector& Position) const;
 	UE_API virtual void InitFromComponent(const UAudioGameplayVolumeComponent* Component);
 
-	UE_API void FindMutatorPriority(FAudioProxyMutatorPriorities& Priorities) const;
-	UE_API void GatherMutators(const FAudioProxyMutatorPriorities& Priorities, FAudioProxyMutatorSearchResult& OutResult) const;
+	UE_API void CollectMutators(ProxyMutatorList& OutMutators) const;
 
 	UE_API void AddPayloadType(PayloadFlags InType);
 	UE_API bool HasPayloadType(PayloadFlags InType) const;
