[AGV] Add opt-in logging to detect which AGVs are most expensive when checking their ContainsPosition()

#rnx
#Synced-CL 48798601

diff --git a/Engine/Plugins/AudioGameplayVolume/Source/AudioGameplayVolume/Private/AudioGameplayVolumeProxy.cpp b/Engine/Plugins/AudioGameplayVolume/Source/AudioGameplayVolume/Private/AudioGameplayVolumeProxy.cpp
--- a/Engine/Plugins/AudioGameplayVolume/Source/AudioGameplayVolume/Private/AudioGameplayVolumeProxy.cpp
+++ b/Engine/Plugins/AudioGameplayVolume/Source/AudioGameplayVolume/Private/AudioGameplayVolumeProxy.cpp
@@ -58,11 +58,6 @@ namespace AudioGameplayVolumeUtils
 
 } // namespace AudioGameplayVolumeUtils
 
-UAudioGameplayVolumeProxy::UAudioGameplayVolumeProxy(const FObjectInitializer& ObjectInitializer)
-	: Super(ObjectInitializer)
-{
-}
-
 bool UAudioGameplayVolumeProxy::ContainsPosition(const FVector& Position) const
 { 
 	return false;
@@ -75,6 +70,10 @@ void UAudioGameplayVolumeProxy::InitFromComponent(const UAudioGameplayVolumeComp
 		UE_LOG(AudioGameplayVolumeLog, Verbose, TEXT("AudioGameplayVolumeProxy - Attempted Init from invalid volume component!"));
 		return;
 	}
+	
+#if !UE_BUILD_SHIPPING
+	AGVOwnerKey = Component->GetOwner();
+#endif
 
 	VolumeID = Component->GetUniqueID();
 	WorldID = Component->GetWorld()->GetUniqueID();
@@ -131,10 +130,12 @@ uint32 UAudioGameplayVolumeProxy::GetWorldID() const
 	return WorldID;
 }
 
-UAGVPrimitiveComponentProxy::UAGVPrimitiveComponentProxy(const FObjectInitializer& ObjectInitializer)
-	: Super(ObjectInitializer)
+#if !UE_BUILD_SHIPPING
+FString UAudioGameplayVolumeProxy::GetOwnerHumanReadableName() const
 {
+	return GetNameSafe(AGVOwnerKey.ResolveObjectPtr());
 }
+#endif
 
 bool UAGVPrimitiveComponentProxy::ContainsPosition(const FVector& Position) const
 {
@@ -197,11 +198,6 @@ void UAGVPrimitiveComponentProxy::InitFromComponent(const UAudioGameplayVolumeCo
 	}
 }
 
-UAGVConditionProxy::UAGVConditionProxy(const FObjectInitializer& ObjectInitializer)
-	: Super(ObjectInitializer)
-{
-}
-
 bool UAGVConditionProxy::ContainsPosition(const FVector& Position) const
 {
 	SCOPED_NAMED_EVENT(UAGVConditionProxy_ContainsPosition, FColor::Blue);
@@ -241,4 +237,3 @@ void UAGVConditionProxy::InitFromComponent(const UAudioGameplayVolumeComponent*
 		}
 	}
 }
-
diff --git a/Engine/Plugins/AudioGameplayVolume/Source/AudioGameplayVolume/Private/AudioGameplayVolumeSubsystem.cpp b/Engine/Plugins/AudioGameplayVolume/Source/AudioGameplayVolume/Private/AudioGameplayVolumeSubsystem.cpp
--- a/Engine/Plugins/AudioGameplayVolume/Source/AudioGameplayVolume/Private/AudioGameplayVolumeSubsystem.cpp
+++ b/Engine/Plugins/AudioGameplayVolume/Source/AudioGameplayVolume/Private/AudioGameplayVolumeSubsystem.cpp
@@ -9,6 +9,7 @@
 #include "AudioGameplayVolumeMutator.h"
 #include "AudioDevice.h"
 #include "Engine/World.h"
+#include "Logging/StructuredLog.h"
 #include "Misc/App.h"
 #include "Stats/Stats.h"
 
@@ -59,6 +60,14 @@ namespace AudioGameplayVolumeConsoleVariables
 		TEXT("When non-zero, will not add a Submix Sends if one already is there sending to the same Submix."),
 		ECVF_Default);
 
+#if !UE_BUILD_SHIPPING
+	float MicroSecsBeforeWarnContainsPositionTooLong = 40.;
+	FAutoConsoleVariableRef CVarMicroSecsBeforeWarnContainsPositionTooLong(
+		TEXT("au.AudioGameplayVolumes.MicroSecsBeforeWarnContainsPositionTooLong"),
+		MicroSecsBeforeWarnContainsPositionTooLong,
+		TEXT("Warns whenever a ContainsPosition test on an AGV would take more than this time, logging the name of the expensive volume."),
+		ECVF_Default);
+#endif
 } // namespace AudioGameplayVolumeConsoleVariables
 
 void FAudioGameplayActiveSoundInfo::Update(double ListenerInteriorStartTime)
@@ -99,11 +108,29 @@ void FAudioProxyMutatorSearchObject::SearchVolumes(const TArray<UAudioGameplayVo
 		{
 			continue;
 		}
-
-		if (!ProxyVolume->ContainsPosition(Location))
+		
+#if !UE_BUILD_SHIPPING
 		{
-			continue;
+			const uint64 StartTime = FPlatformTime::Cycles64();
+#endif
+
+			if (!ProxyVolume->ContainsPosition(Location))
+			{
+				continue;
+			}
+			
+#if !UE_BUILD_SHIPPING
+			const double DurationMicroSecs = FTimespan(FPlatformTime::Cycles64() - StartTime).GetTotalMicroseconds();
+			if (DurationMicroSecs >= AudioGameplayVolumeConsoleVariables::MicroSecsBeforeWarnContainsPositionTooLong)
+			{
+				// Some AGV proxies might just have costly ContainsPosition() checks, or Brush-based volumes might just have a shape that's too 
+				// complex. This allows to find outliers.
+				// Note: With WorldPartition enabled and its UAID_ actor names, you can still search your level in the Editor
+				//       Outliner for that string and the volume actor will show up.
+				UE_LOGFMT(AudioGameplayVolumeLog, Verbose, "{AGV} ContainsPosition() took {Duration} µs", ProxyVolume->GetOwnerHumanReadableName(), DurationMicroSecs);
+			}
 		}
+#endif
 
 		if (bCollectMutators)
 		{
diff --git a/Engine/Plugins/AudioGameplayVolume/Source/AudioGameplayVolume/Public/AudioGameplayVolumeProxy.h b/Engine/Plugins/AudioGameplayVolume/Source/AudioGameplayVolume/Public/AudioGameplayVolumeProxy.h
--- a/Engine/Plugins/AudioGameplayVolume/Source/AudioGameplayVolume/Public/AudioGameplayVolumeProxy.h
+++ b/Engine/Plugins/AudioGameplayVolume/Source/AudioGameplayVolume/Public/AudioGameplayVolumeProxy.h
@@ -3,6 +3,8 @@
 #pragma once
 
 #include "AudioGameplayFlags.h"
+#include "UObject/ObjectKey.h"
+
 #include "AudioGameplayVolumeProxy.generated.h"
 
 #define UE_API AUDIOGAMEPLAYVOLUME_API
@@ -26,7 +28,7 @@ struct FBodyInstance;
 UCLASS(MinimalAPI, Abstract, EditInlineNew, HideDropdown)
 class UAudioGameplayVolumeProxy : public UObject
 {
-	GENERATED_UCLASS_BODY()
+	GENERATED_BODY()
 
 public:
 
@@ -56,6 +58,15 @@ protected:
 	uint32 VolumeID = INDEX_NONE;
 	uint32 WorldID = INDEX_NONE;
 	PayloadFlags PayloadType = PayloadFlags::AGCP_None;
+
+#if !UE_BUILD_SHIPPING
+public:
+	FString GetOwnerHumanReadableName() const;
+	
+private:
+	// Only for profiling use
+	FObjectKey AGVOwnerKey;
+#endif
 };
 
 /**
@@ -68,7 +79,7 @@ protected:
 UCLASS(MinimalAPI, meta = (DisplayName = "Audio Listener in Primitives"))
 class UAGVPrimitiveComponentProxy : public UAudioGameplayVolumeProxy
 {
-	GENERATED_UCLASS_BODY()
+	GENERATED_BODY()
 
 public:
 
@@ -95,7 +106,7 @@ protected:
 UCLASS(MinimalAPI, meta = (DisplayName = "Arbitrary"))
 class UAGVConditionProxy : public UAudioGameplayVolumeProxy
 {
-	GENERATED_UCLASS_BODY()
+	GENERATED_BODY()
 
 public:
 
