[AGV] Fix early distance culling not working for beam-like primitives: now uses bounding box rather than bounding sphere

Beam-like shape means a shape where the magnitude of one or its XYZ component is much larger than the others. For these, using a bounding sphere would not optimize cull early properly, especially if users create "infinitely tall shapes" effectively being "2D zones" on maps: no early culling can happen because we are always in the bounding sphere of these infinitely tall shapes.

Also fix costly volumes logging ignoring volumes we are inside of.

#Synced-CL 49316309

diff --git a/Engine/Plugins/AudioGameplayVolume/Source/AudioGameplayVolume/Private/AudioGameplayVolumeProxy.cpp b/Engine/Plugins/AudioGameplayVolume/Source/AudioGameplayVolume/Private/AudioGameplayVolumeProxy.cpp
--- a/Engine/Plugins/AudioGameplayVolume/Source/AudioGameplayVolume/Private/AudioGameplayVolumeProxy.cpp
+++ b/Engine/Plugins/AudioGameplayVolume/Source/AudioGameplayVolume/Private/AudioGameplayVolumeProxy.cpp
@@ -50,10 +50,7 @@ namespace AudioGameplayVolumeUtils
 		}
 
 		// Early distance culling
-		const float BoundsRadiusSq = FMath::Square(PrimitiveComponent->Bounds.SphereRadius);
-		const float DistanceSq = FVector::DistSquared(PrimitiveComponent->Bounds.Origin, Position);
-
-		return DistanceSq <= BoundsRadiusSq;
+		return PrimitiveComponent->Bounds.GetBox().IsInsideOrOn(Position);
 	}
 
 } // namespace AudioGameplayVolumeUtils
@@ -131,9 +128,9 @@ uint32 UAudioGameplayVolumeProxy::GetWorldID() const
 }
 
 #if !UE_BUILD_SHIPPING
-FString UAudioGameplayVolumeProxy::GetOwnerHumanReadableName() const
+FName UAudioGameplayVolumeProxy::GetOwnerName() const
 {
-	return GetNameSafe(AGVOwnerKey.ResolveObjectPtr());
+	return GetFNameSafe(AGVOwnerKey.ResolveObjectPtr());
 }
 #endif
 
diff --git a/Engine/Plugins/AudioGameplayVolume/Source/AudioGameplayVolume/Private/AudioGameplayVolumeSubsystem.cpp b/Engine/Plugins/AudioGameplayVolume/Source/AudioGameplayVolume/Private/AudioGameplayVolumeSubsystem.cpp
--- a/Engine/Plugins/AudioGameplayVolume/Source/AudioGameplayVolume/Private/AudioGameplayVolumeSubsystem.cpp
+++ b/Engine/Plugins/AudioGameplayVolume/Source/AudioGameplayVolume/Private/AudioGameplayVolumeSubsystem.cpp
@@ -109,28 +109,27 @@ void FAudioProxyMutatorSearchObject::SearchVolumes(const TArray<UAudioGameplayVo
 			continue;
 		}
 		
-#if !UE_BUILD_SHIPPING
 		{
+#if !UE_BUILD_SHIPPING
 			const uint64 StartTime = FPlatformTime::Cycles64();
+			ON_SCOPE_EXIT{
+				const double DurationMicroSecs = FTimespan(FPlatformTime::Cycles64() - StartTime).GetTotalMicroseconds();
+				if (DurationMicroSecs >= AudioGameplayVolumeConsoleVariables::MicroSecsBeforeWarnContainsPositionTooLong)
+				{
+					// Some AGV proxies might just have costly ContainsPosition() checks, or Brush-based volumes might just have a shape that's too 
+					// complex. This allows to find outliers.
+					// Note: With WorldPartition enabled and its UAID_ actor names, you can still search your level in the Editor
+					//       Outliner for that string and the volume actor will show up.
+					UE_LOGFMT(AudioGameplayVolumeLog, Verbose, "{AGV} ContainsPosition() took {Duration} µs", ProxyVolume->GetOwnerName(), DurationMicroSecs);
+				}
+			};
 #endif
 
 			if (!ProxyVolume->ContainsPosition(Location))
 			{
 				continue;
 			}
-			
-#if !UE_BUILD_SHIPPING
-			const double DurationMicroSecs = FTimespan(FPlatformTime::Cycles64() - StartTime).GetTotalMicroseconds();
-			if (DurationMicroSecs >= AudioGameplayVolumeConsoleVariables::MicroSecsBeforeWarnContainsPositionTooLong)
-			{
-				// Some AGV proxies might just have costly ContainsPosition() checks, or Brush-based volumes might just have a shape that's too 
-				// complex. This allows to find outliers.
-				// Note: With WorldPartition enabled and its UAID_ actor names, you can still search your level in the Editor
-				//       Outliner for that string and the volume actor will show up.
-				UE_LOGFMT(AudioGameplayVolumeLog, Verbose, "{AGV} ContainsPosition() took {Duration} µs", ProxyVolume->GetOwnerHumanReadableName(), DurationMicroSecs);
-			}
 		}
-#endif
 
 		if (bCollectMutators)
 		{
diff --git a/Engine/Plugins/AudioGameplayVolume/Source/AudioGameplayVolume/Public/AudioGameplayVolumeProxy.h b/Engine/Plugins/AudioGameplayVolume/Source/AudioGameplayVolume/Public/AudioGameplayVolumeProxy.h
--- a/Engine/Plugins/AudioGameplayVolume/Source/AudioGameplayVolume/Public/AudioGameplayVolumeProxy.h
+++ b/Engine/Plugins/AudioGameplayVolume/Source/AudioGameplayVolume/Public/AudioGameplayVolumeProxy.h
@@ -61,7 +61,7 @@ protected:
 
 #if !UE_BUILD_SHIPPING
 public:
-	FString GetOwnerHumanReadableName() const;
+	FName GetOwnerName() const;
 	
 private:
 	// Only for profiling use
