Fix race condition in AudioGameplayVolumeSubsystem

#rnx
#Synced-CL 49560450

diff --git a/Engine/Plugins/AudioGameplayVolume/Source/AudioGameplayVolume/Private/AudioGameplayVolumeSubsystem.cpp b/Engine/Plugins/AudioGameplayVolume/Source/AudioGameplayVolume/Private/AudioGameplayVolumeSubsystem.cpp
--- a/Engine/Plugins/AudioGameplayVolume/Source/AudioGameplayVolume/Private/AudioGameplayVolumeSubsystem.cpp
+++ b/Engine/Plugins/AudioGameplayVolume/Source/AudioGameplayVolume/Private/AudioGameplayVolumeSubsystem.cpp
@@ -11,6 +11,7 @@
 #include "Engine/World.h"
 #include "Logging/StructuredLog.h"
 #include "Misc/App.h"
+#include "Misc/ScopeLock.h"
 #include "Stats/Stats.h"
 
 #include UE_INLINE_GENERATED_CPP_BY_NAME(AudioGameplayVolumeSubsystem)
@@ -227,6 +228,7 @@ bool UAudioGameplayVolumeSubsystem::ShouldCreateSubsystem(UObject* Outer) const
 
 void UAudioGameplayVolumeSubsystem::Deinitialize()
 {
+	FScopeLock Lock(&AGVComponentsCriticalSection);
 	AGVComponents.Reset();
 }
 
@@ -430,7 +432,10 @@ void UAudioGameplayVolumeSubsystem::AddVolumeComponent(const UAudioGameplayVolum
 	if (VolumeProxy && !AGVComponents.Contains(ComponentID))
 	{
 		VolumeProxy->InitFromComponent(VolumeComponent);
-		AGVComponents.Emplace(ComponentID, VolumeComponent);
+		{
+			FScopeLock Lock(&AGVComponentsCriticalSection);
+			AGVComponents.Emplace(ComponentID, VolumeComponent);
+		}
 		const uint32 WorldID = VolumeProxy->GetWorldID();
 
 		UE_LOG(AudioGameplayVolumeLog, VeryVerbose, TEXT("AudioGameplayVolumeComponent %s [%08x] added"), *VolumeComponent->GetFName().ToString(), ComponentID);
@@ -500,6 +505,7 @@ void UAudioGameplayVolumeSubsystem::RemoveVolumeComponent(const UAudioGameplayVo
 	const uint32 ComponentID = VolumeComponent->GetUniqueID();
 	if (AGVComponents.Contains(ComponentID))
 	{
+		FScopeLock Lock(&AGVComponentsCriticalSection);
 		AGVComponents.Remove(ComponentID);
 		UE_LOG(AudioGameplayVolumeLog, VeryVerbose, TEXT("AudioGameplayVolumeComponent %s [%08x] removed"), *VolumeComponent->GetFName().ToString(), ComponentID);
 	}
@@ -549,14 +555,20 @@ bool UAudioGameplayVolumeSubsystem::AddProxy(uint32 AudioGameplayVolumeID, uint3
 {
 	check(IsInAudioThread());
 
+	const UAudioGameplayVolumeComponent* AGVComponent;
+	{
+		FScopeLock Lock(&AGVComponentsCriticalSection);
+		AGVComponent = AGVComponents.FindRef(AudioGameplayVolumeID);
+	}
+
 	// Make sure the component is in the map and non-null
-	if (!AGVComponents.Contains(AudioGameplayVolumeID) || AGVComponents[AudioGameplayVolumeID] == nullptr)
+	if (AGVComponent == nullptr)
 	{
 		return false;
 	}
 
 	// Make sure our proxy is valid
-	if (AGVComponents[AudioGameplayVolumeID]->GetProxy() == nullptr)
+	if (AGVComponent->GetProxy() == nullptr)
 	{
 		return false;
 	}
@@ -606,6 +618,8 @@ bool UAudioGameplayVolumeSubsystem::RemoveProxy(uint32 AudioGameplayVolumeID)
 void UAudioGameplayVolumeSubsystem::GenerateVolumeProxyList()
 {
 	TransientProxyList.Reset(256);
+	
+	FScopeLock Lock(&AGVComponentsCriticalSection);
 	Algo::Transform(AGVComponents, TransientProxyList, [](const TPair<uint32, TObjectPtr<const UAudioGameplayVolumeComponent>>& Pair)
 	{
 		return Pair.Value ? Pair.Value->GetProxy() : nullptr;
diff --git a/Engine/Plugins/AudioGameplayVolume/Source/AudioGameplayVolume/Public/AudioGameplayVolumeSubsystem.h b/Engine/Plugins/AudioGameplayVolume/Source/AudioGameplayVolume/Public/AudioGameplayVolumeSubsystem.h
--- a/Engine/Plugins/AudioGameplayVolume/Source/AudioGameplayVolume/Public/AudioGameplayVolumeSubsystem.h
+++ b/Engine/Plugins/AudioGameplayVolume/Source/AudioGameplayVolume/Public/AudioGameplayVolumeSubsystem.h
@@ -172,10 +172,13 @@ protected:
 	UPROPERTY(Transient)
 	TArray<TObjectPtr<UAudioGameplayVolumeProxy>> TransientProxyList;
 
-	// Components in our system
+	// Components in our system. This map should only be modified by the game thread.
 	UPROPERTY(Transient)
 	TMap<uint32, TObjectPtr<const UAudioGameplayVolumeComponent>> AGVComponents;
 
+	// Used to protect concurrent access to the AGVComponents map.
+	FCriticalSection AGVComponentsCriticalSection;
+
 	// List of Component IDs (for proxies) the audio thread is tracking
 	TArray<uint32> KnownProxyIDs;
 
