[AudioGameplayVolumes] Fix priority system for listener-based mutators (e.g. SubmixOverride)

#tests Confirmed in Test_AGV map that this fixes the bug
#rnx
#Synced-CL 52209130

diff --git a/Engine/Plugins/AudioGameplayVolume/Source/AudioGameplayVolume/Private/AudioGameplayVolumeListener.cpp b/Engine/Plugins/AudioGameplayVolume/Source/AudioGameplayVolume/Private/AudioGameplayVolumeListener.cpp
--- a/Engine/Plugins/AudioGameplayVolume/Source/AudioGameplayVolume/Private/AudioGameplayVolumeListener.cpp
+++ b/Engine/Plugins/AudioGameplayVolume/Source/AudioGameplayVolume/Private/AudioGameplayVolumeListener.cpp
@@ -5,6 +5,18 @@
 #include "AudioGameplayVolumeSubsystem.h"
 #include "Misc/App.h"
 
+namespace AudioGameplayVolumeListenerCVars
+{
+	int32 bFilterListenerMutatorsByPriority = 1;
+	FAutoConsoleVariableRef CVarFilterListenerMutatorsByPriority(
+		TEXT("au.AudioGameplayVolumes.FilterListenerMutatorsByPriority"),
+		bFilterListenerMutatorsByPriority,
+		TEXT("When non-zero, listener-applied mutators respect priority (only mutators with the highest priority level of each type are applied).\n")
+		TEXT("When zero, uses legacy behavior where all mutators are applied regardless of priority.\n")
+		TEXT("0: Disable (legacy), 1: Enable (default)"),
+		ECVF_Default);
+} // namespace AudioGameplayVolumeListenerCVars
+
 void FInterpolatedInteriorSettings::Apply(const FInteriorSettings& NewSettings)
 {
 	if (InteriorSettings != NewSettings)
@@ -82,11 +94,38 @@ void FAudioGameplayVolumeListener::Update(const FAudioProxyMutatorSearchResult&
 			Mutator->Remove(*this);
 		}
 
-		// Now cache and apply the currently matching mutators
+		// Now cache and apply the currently matching mutators.
+		// Mutators are sorted by type (name) and descending priority, so we can quickly apply the highest priority mutators.
 		ActiveMutators = Result.MatchingMutators;
-		for (const TSharedPtr<FProxyVolumeMutator>& Mutator : ActiveMutators)
+		if (AudioGameplayVolumeListenerCVars::bFilterListenerMutatorsByPriority)
+		{
+			FName MutatorGroup;
+			int32 CurrentPriority = INDEX_NONE;
+			for (const TSharedPtr<FProxyVolumeMutator>& Mutator : ActiveMutators)
+			{
+				// If we encounter a new name, it's a new type - with the highest priority for that group
+				if (!Mutator->MutatorName.IsEqual(MutatorGroup))
+				{
+					MutatorGroup = Mutator->MutatorName;
+					CurrentPriority = Mutator->Priority;
+				}
+
+				// Lower priority mutators are ignored, we need to advance to the next named group
+				if (Mutator->Priority < CurrentPriority)
+				{
+					continue;
+				}
+
+				Mutator->Apply(*this);
+			}
+		}
+		else
 		{
-			Mutator->Apply(*this);
+			// Legacy behavior which will be removed -- apply all mutators in reverse priority order.
+			for (const TSharedPtr<FProxyVolumeMutator>& Mutator : ActiveMutators)
+			{
+				Mutator->Apply(*this);
+			}
 		}
 
 		// Reapply interior from mutators
diff --git a/Engine/Plugins/AudioGameplayVolume/Source/AudioGameplayVolume/Private/SubmixOverrideVolumeComponent.cpp b/Engine/Plugins/AudioGameplayVolume/Source/AudioGameplayVolume/Private/SubmixOverrideVolumeComponent.cpp
--- a/Engine/Plugins/AudioGameplayVolume/Source/AudioGameplayVolume/Private/SubmixOverrideVolumeComponent.cpp
+++ b/Engine/Plugins/AudioGameplayVolume/Source/AudioGameplayVolume/Private/SubmixOverrideVolumeComponent.cpp
@@ -2,6 +2,7 @@
 
 #include "SubmixOverrideVolumeComponent.h"
 #include "AudioDevice.h"
+#include "AudioGameplayVolumeLogs.h"
 #include "AudioGameplayVolumeListener.h"
 #include "Sound/SoundEffectSubmix.h"
 #include "Sound/SoundSubmix.h"
@@ -31,11 +32,16 @@ void FProxyMutator_SubmixOverride::Apply(FAudioGameplayVolumeListener& Listener)
 		{
 			if (OverrideSettings.Submix && OverrideSettings.SubmixEffectChain.Num() > 0)
 			{
+				UE_LOG(AudioGameplayVolumeLog, Verbose, TEXT("SubmixOverride Apply - Submix: %s, Priority: %d, VolumeID: %08x, EffectChain: %d effects, CrossfadeTime: %.3f"),
+					*OverrideSettings.Submix->GetName(), Priority, VolumeID, OverrideSettings.SubmixEffectChain.Num(), OverrideSettings.CrossfadeTime);
+
 				// Build the instances of the new submix preset chain override
 				for (USoundEffectSubmixPreset* SubmixEffectPreset : OverrideSettings.SubmixEffectChain)
 				{
 					if (SubmixEffectPreset)
 					{
+						UE_LOG(AudioGameplayVolumeLog, Verbose, TEXT("  Effect: %s"), *SubmixEffectPreset->GetName());
+
 						InitData.PresetSettings = nullptr;
 						InitData.ParentPresetUniqueId = SubmixEffectPreset->GetUniqueID();
 
@@ -62,6 +68,9 @@ void FProxyMutator_SubmixOverride::Remove(FAudioGameplayVolumeListener& Listener
 		// Clear out any previous submix effect chain overrides
 		for (const FAudioVolumeSubmixOverrideSettings& OverrideSettings : SubmixOverrideSettings)
 		{
+			UE_LOG(AudioGameplayVolumeLog, Verbose, TEXT("SubmixOverride Remove - Submix: %s, Priority: %d, VolumeID: %08x, CrossfadeTime: %.3f"),
+				OverrideSettings.Submix ? *OverrideSettings.Submix->GetName() : TEXT("null"), Priority, VolumeID, OverrideSettings.CrossfadeTime);
+
 			AudioDeviceHandle->ClearSubmixEffectChainOverride(OverrideSettings.Submix, OverrideSettings.CrossfadeTime);
 		}
 	}
